PerformanceCaptureCore plugin

-Retarget anim instance using IKRetarget asset

-Actor component to force use of RetargetAnimInstance on chosen skeletal mesh component
-Specific Character Skeletal Mesh actor that implements the Retarget component with BP functions for setting source performer, retarget asset, custom retarget profile

-Actor component to force use of LiveLinkInstance to drive motion capture performers
-Specfic Performer SkeletalMesh actor that implements the Performer component with BP functions for getting/setting LiveLink subject, pause animation

diff --git a/Engine/Plugins/Animation/PerformanceCaptureCore/PerformanceCaptureCore.uplugin b/Engine/Plugins/Animation/PerformanceCaptureCore/PerformanceCaptureCore.uplugin
new file mode 100644
--- /dev/null
+++ b/Engine/Plugins/Animation/PerformanceCaptureCore/PerformanceCaptureCore.uplugin
@@ -0,0 +1,34 @@
+﻿{
+	"FileVersion": 3,
+	"Version": 1,
+	"VersionName": "1.0",
+	"FriendlyName": "Performance Capture Core",
+	"Description": "Performance Capture Core Actor and Component Classes ",
+	"Category": "Animation",
+	"CreatedBy": "Epic Games",
+	"CreatedByURL": "",
+	"DocsURL": "",
+	"MarketplaceURL": "",
+	"CanContainContent": false,
+	"IsBetaVersion": false,
+	"IsExperimentalVersion": true,
+	"Installed": false,
+	"Modules": [
+		{
+			"Name": "PerformanceCaptureCore",
+			"Type": "Runtime",
+			"LoadingPhase": "Default"
+		}
+	],
+	"Plugins": [
+		{
+			"Name": "IKRig",
+			"Enabled": true
+		},
+		{
+			"Name": "LiveLink",
+			"Enabled": true
+		}
+	],
+	"SupportURL": ""
+}
\ No newline at end of file
diff --git a/Engine/Plugins/Animation/PerformanceCaptureCore/Source/PerformanceCaptureCore/PerformanceCaptureCore.Build.cs b/Engine/Plugins/Animation/PerformanceCaptureCore/Source/PerformanceCaptureCore/PerformanceCaptureCore.Build.cs
new file mode 100644
--- /dev/null
+++ b/Engine/Plugins/Animation/PerformanceCaptureCore/Source/PerformanceCaptureCore/PerformanceCaptureCore.Build.cs
@@ -0,0 +1,31 @@
+﻿// Copyright Epic Games, Inc. All Rights Reserved.
+
+using UnrealBuildTool;
+
+public class PerformanceCaptureCore : ModuleRules
+{
+	public PerformanceCaptureCore(ReadOnlyTargetRules Target) : base(Target)
+	{
+		PCHUsage = ModuleRules.PCHUsageMode.UseExplicitOrSharedPCHs;
+		
+		PublicDependencyModuleNames.AddRange(
+			new string[]
+			{
+				"Core",
+			}
+			);
+		
+		PrivateDependencyModuleNames.AddRange(
+			new string[]
+			{
+				"CoreUObject",
+				"Engine",
+				"Slate",
+				"SlateCore",
+				"IKRig",
+				"LiveLinkAnimationCore",
+				"LiveLinkInterface"
+			}
+			);
+	}
+}
diff --git a/Engine/Plugins/Animation/PerformanceCaptureCore/Source/PerformanceCaptureCore/Private/CaptureCharacter.cpp b/Engine/Plugins/Animation/PerformanceCaptureCore/Source/PerformanceCaptureCore/Private/CaptureCharacter.cpp
new file mode 100644
--- /dev/null
+++ b/Engine/Plugins/Animation/PerformanceCaptureCore/Source/PerformanceCaptureCore/Private/CaptureCharacter.cpp
@@ -0,0 +1,100 @@
+// Copyright Epic Games, Inc. All Rights Reserved.
+
+#include "CaptureCharacter.h"
+
+#include "RetargetAnimInstance.h"
+#include "RetargetComponent.h"
+
+// Sets default values
+ACaptureCharacter::ACaptureCharacter()
+{
+	// Set this actor to call Tick() every frame.  You can turn this off to improve performance if you don't need it.
+	PrimaryActorTick.bCanEverTick = true;
+	RetargetComponent = CreateDefaultSubobject<URetargetComponent>("RetargetComponent");
+	RetargetComponent->ControlledSkeletalMeshComponent.OverrideComponent = Cast<USkeletalMeshComponent>(GetRootComponent());
+	RetargetComponent->SourceSkeletalMeshComponent.OverrideComponent = nullptr;
+}
+
+
+//Update the custom retarget profile
+void ACaptureCharacter::SetCustomRetargetProfile(FRetargetProfile InProfile)
+{
+	RetargetComponent->SetCustomRetargetProfile(InProfile);
+}
+
+//Get the custom retarget profile
+FRetargetProfile ACaptureCharacter::GetCustomRetargetProfile()
+{
+	FRetargetProfile OutProfile = RetargetComponent->GetCustomRetargetProfile();
+	return OutProfile;
+}
+
+void ACaptureCharacter::SetSourcePerformer(ACapturePerformer* InPerformer)
+{
+	if(IsValid(InPerformer))
+	{
+		SourcePerformer = InPerformer;
+		RetargetComponent->SetSourcePerformerMesh(InPerformer->GetSkeletalMeshComponent());
+		RetargetComponent->InitiateAnimation();
+	}
+	else
+	{
+		SourcePerformer = nullptr;
+		RetargetComponent->SetSourcePerformerMesh(GetSkeletalMeshComponent());
+		RetargetComponent->InitiateAnimation();
+	}
+}
+
+//Update the retarget asset
+void ACaptureCharacter::SetRetargetAsset(UIKRetargeter* InRetargetAsset)
+{
+	RetargetAsset = InRetargetAsset;
+	RetargetComponent->SetRetargetAsset(InRetargetAsset);
+}
+
+void ACaptureCharacter::SetForceAllSkeletalMeshesToFollowLeader(bool InFollowLeader)
+{
+	bForceAllSkeletalMeshesToFollowLeader = InFollowLeader;
+	RetargetComponent->bForceOtherMeshesToFollowControlledMesh = bForceAllSkeletalMeshesToFollowLeader;
+}
+
+void ACaptureCharacter::PostRegisterAllComponents()
+{
+	Super::PostRegisterAllComponents();
+	if (!HasAnyFlags(RF_ClassDefaultObject))
+	{
+		RetargetComponent->SetControlledMesh(GetSkeletalMeshComponent());
+		SetSourcePerformer(SourcePerformer);
+		RetargetComponent->SetRetargetAsset(RetargetAsset);
+		RetargetComponent->SetForceOtherMeshesToFollowControlledMesh(bForceAllSkeletalMeshesToFollowLeader);
+		RetargetComponent->InitiateAnimation();
+	}
+}
+
+#if WITH_EDITOR
+void ACaptureCharacter::PostEditChangeProperty(FPropertyChangedEvent& PropertyChangedEvent)
+{
+	Super::PostEditChangeProperty(PropertyChangedEvent);
+
+	if(PropertyChangedEvent.Property != nullptr)
+	{
+		if(PropertyChangedEvent.Property->GetFName() == GET_MEMBER_NAME_CHECKED(ACaptureCharacter, SourcePerformer))
+		{
+			SetSourcePerformer(SourcePerformer);
+			RetargetComponent->InitiateAnimation();
+		}
+		
+		if(PropertyChangedEvent.Property->GetFName() == GET_MEMBER_NAME_CHECKED(ACaptureCharacter, RetargetAsset))
+		{
+			SetRetargetAsset(RetargetAsset);
+			RetargetComponent->InitiateAnimation();
+		}
+		
+		if(PropertyChangedEvent.Property->GetFName() == GET_MEMBER_NAME_CHECKED(ACaptureCharacter, bForceAllSkeletalMeshesToFollowLeader))
+		{
+			RetargetComponent->bForceOtherMeshesToFollowControlledMesh = bForceAllSkeletalMeshesToFollowLeader;
+			RetargetComponent->InitiateAnimation();
+		}
+	}
+}
+#endif
\ No newline at end of file
diff --git a/Engine/Plugins/Animation/PerformanceCaptureCore/Source/PerformanceCaptureCore/Private/CapturePerformer.cpp b/Engine/Plugins/Animation/PerformanceCaptureCore/Source/PerformanceCaptureCore/Private/CapturePerformer.cpp
new file mode 100644
--- /dev/null
+++ b/Engine/Plugins/Animation/PerformanceCaptureCore/Source/PerformanceCaptureCore/Private/CapturePerformer.cpp
@@ -0,0 +1,103 @@
+// Copyright Epic Games, Inc. All Rights Reserved.
+
+#include "CapturePerformer.h"
+#include "PerformerComponent.h"
+
+// Sets default values
+ACapturePerformer::ACapturePerformer()
+{
+	// Set this actor to call Tick() every frame.  You can turn this off to improve performance if you don't need it.
+	PrimaryActorTick.bCanEverTick = true;
+
+	PerformerComponent = CreateDefaultSubobject<UPerformerComponent>("PerformerComponent");
+	PerformerComponent->ControlledSkeletalMesh.OverrideComponent = Cast<USkeletalMeshComponent>(GetRootComponent());
+}
+
+// Called when the game starts or when spawned
+void ACapturePerformer::BeginPlay()
+{
+	Super::BeginPlay();
+}
+
+void ACapturePerformer::PostInitializeComponents()
+{
+	Super::PostInitializeComponents();
+}
+
+void ACapturePerformer::PostRegisterAllComponents()
+{
+	Super::PostRegisterAllComponents();
+	
+	if (!HasAnyFlags(RF_ClassDefaultObject))
+	{
+		PerformerComponent->ControlledSkeletalMesh.OverrideComponent = GetSkeletalMeshComponent();
+		PerformerComponent->SetLiveLinkSubject(SubjectName);
+		PerformerComponent->bForceOtherMeshesToFollowControlledMesh = bForceAllMeshesToFollowLeader;
+		PerformerComponent->InitiateAnimation();
+	}
+}
+
+void ACapturePerformer::OnConstruction(const FTransform& Transform)
+{
+	Super::OnConstruction(Transform);
+}
+
+#if WITH_EDITOR
+void ACapturePerformer::PostEditChangeProperty(FPropertyChangedEvent& PropertyChangedEvent)
+{
+	Super::PostEditChangeProperty(PropertyChangedEvent);
+
+	if(PropertyChangedEvent.Property != nullptr)
+	{
+		if(PropertyChangedEvent.Property->GetFName() == GET_MEMBER_NAME_CHECKED(ACapturePerformer, SubjectName))
+		{
+			PerformerComponent->InitiateAnimation();
+			PerformerComponent->SetLiveLinkSubject(SubjectName);
+		}
+		if(PropertyChangedEvent.Property->GetFName() == GET_MEMBER_NAME_CHECKED(ACapturePerformer, bForceAllMeshesToFollowLeader))
+		{
+			PerformerComponent->bForceOtherMeshesToFollowControlledMesh = bForceAllMeshesToFollowLeader;
+			PerformerComponent->InitiateAnimation();
+		}
+		if(PropertyChangedEvent.Property->GetFName() == GET_MEMBER_NAME_CHECKED(ACapturePerformer, bEvaluateAnimation))
+		{
+			SetEvaluateLiveLinkData(bEvaluateAnimation);
+		}
+	}
+}
+#endif
+
+void ACapturePerformer::SetLiveLinkSubject(FLiveLinkSubjectName Subject)
+{
+	SubjectName = Subject;
+	PerformerComponent->SetLiveLinkSubject(Subject);
+}
+
+FLiveLinkSubjectName ACapturePerformer::GetLiveLinkSubject() const
+{
+	return PerformerComponent->SubjectName;
+}
+
+void ACapturePerformer::SetEvaluateLiveLinkData(bool bEvaluateLinkLink)
+{
+	bEvaluateAnimation = bEvaluateLinkLink;
+	PerformerComponent->SetEvaluateLiveLinkData(bEvaluateLinkLink);
+}
+
+bool ACapturePerformer::GetEvaluateLiveLinkData()
+{
+	USkeletalMeshComponent* ControlledMesh;
+	ControlledMesh = Cast<USkeletalMeshComponent>(PerformerComponent->ControlledSkeletalMesh.GetComponent(GetOwner()));
+	ULiveLinkInstance* LiveLinkAnimInstance = Cast<ULiveLinkInstance>(ControlledMesh->GetAnimInstance());
+	
+	if (IsValid(LiveLinkAnimInstance) && IsValid(ControlledMesh))
+	{
+		return LiveLinkAnimInstance->GetEnableLiveLinkEvaluation();
+	}
+	return false;
+}
+
+void ACapturePerformer::SetMocapMesh(USkeletalMesh* MocapMesh)
+{
+	GetSkeletalMeshComponent()->SetSkeletalMeshAsset(MocapMesh);	
+}
diff --git a/Engine/Plugins/Animation/PerformanceCaptureCore/Source/PerformanceCaptureCore/Private/PerformanceCaptureCore.cpp b/Engine/Plugins/Animation/PerformanceCaptureCore/Source/PerformanceCaptureCore/Private/PerformanceCaptureCore.cpp
new file mode 100644
--- /dev/null
+++ b/Engine/Plugins/Animation/PerformanceCaptureCore/Source/PerformanceCaptureCore/Private/PerformanceCaptureCore.cpp
@@ -0,0 +1,20 @@
+﻿// Copyright Epic Games, Inc. All Rights Reserved.
+
+#include "PerformanceCaptureCore.h"
+
+#define LOCTEXT_NAMESPACE "FPerformanceCaptureCoreModule"
+
+void FPerformanceCaptureCoreModule::StartupModule()
+{
+	// This code will execute after your module is loaded into memory; the exact timing is specified in the .uplugin file per-module
+}
+
+void FPerformanceCaptureCoreModule::ShutdownModule()
+{
+	// This function may be called during shutdown to clean up your module.  For modules that support dynamic reloading,
+	// we call this function before unloading the module.
+}
+
+#undef LOCTEXT_NAMESPACE
+	
+IMPLEMENT_MODULE(FPerformanceCaptureCoreModule, PerformanceCaptureCore)
\ No newline at end of file
diff --git a/Engine/Plugins/Animation/PerformanceCaptureCore/Source/PerformanceCaptureCore/Private/PerformerComponent.cpp b/Engine/Plugins/Animation/PerformanceCaptureCore/Source/PerformanceCaptureCore/Private/PerformerComponent.cpp
new file mode 100644
--- /dev/null
+++ b/Engine/Plugins/Animation/PerformanceCaptureCore/Source/PerformanceCaptureCore/Private/PerformerComponent.cpp
@@ -0,0 +1,187 @@
+// Copyright Epic Games, Inc. All Rights Reserved.
+
+#include "PerformerComponent.h"
+#include "LiveLinkInstance.h"
+
+// Sets default values for this component's properties
+UPerformerComponent::UPerformerComponent()
+	:bIsDirty(true)
+{
+	// Set this component to be initialized when the game starts, and to be ticked every frame. 
+	PrimaryComponentTick.bCanEverTick = true;
+	PrimaryComponentTick.bStartWithTickEnabled = true;
+	bTickInEditor = true;
+}
+
+void UPerformerComponent::DestroyComponent(bool bPromoteChildren)
+{
+	//Reset all the skeletal meshes in the owner actor 
+	TArray<USkeletalMeshComponent*> SkeletalMeshComponents;
+	GetOwner()->GetComponents(SkeletalMeshComponents);
+	for (USkeletalMeshComponent* SkeletalMeshComponent : SkeletalMeshComponents)
+	{
+		SkeletalMeshComponent->SetUpdateAnimationInEditor(false/* NewUpdateState */);
+		SkeletalMeshComponent->SetAnimClass(nullptr);
+		SkeletalMeshComponent->InitAnim(true/* bForceReinit */);
+	}
+	
+	Super::DestroyComponent(bPromoteChildren);
+}
+
+void UPerformerComponent::OnRegister()
+{
+	Super::OnRegister();
+	
+	if(bIsDirty)
+	{
+		InitiateAnimation();
+	}
+}
+
+#if WITH_EDITOR
+void UPerformerComponent::PostEditChangeProperty(FPropertyChangedEvent& PropertyChangedEvent)
+{
+	Super::PostEditChangeProperty(PropertyChangedEvent);
+
+	if(PropertyChangedEvent.Property != nullptr)
+	{
+		if (PropertyChangedEvent.Property->GetFName() == GET_MEMBER_NAME_CHECKED(UPerformerComponent, SubjectName))
+		{
+			SetLiveLinkSubject(SubjectName);
+			bIsDirty=true;
+		}
+	
+		if (PropertyChangedEvent.Property->GetFName() == GET_MEMBER_NAME_CHECKED(UPerformerComponent, bEvaluateAnimation))
+		{
+			SetEvaluateLiveLinkData(bEvaluateAnimation);
+		}
+	
+		if (PropertyChangedEvent.Property->GetFName() == GET_MEMBER_NAME_CHECKED(UPerformerComponent, ControlledSkeletalMesh))
+		{
+			bIsDirty=true;
+		}
+	
+		if (PropertyChangedEvent.Property->GetFName() == GET_MEMBER_NAME_CHECKED(UPerformerComponent, bForceOtherMeshesToFollowControlledMesh))
+		{
+			bIsDirty=true;
+		}
+	}
+}
+#endif
+
+// Called every frame
+void UPerformerComponent::TickComponent(float DeltaTime, ELevelTick TickType, FActorComponentTickFunction* ThisTickFunction)
+{
+	Super::TickComponent(DeltaTime, TickType, ThisTickFunction);
+	
+	if(bIsDirty)
+	{
+		InitiateAnimation();
+	}
+}
+
+void UPerformerComponent::SetLiveLinkSubject(FLiveLinkSubjectName Subject)
+{
+	SubjectName = Subject;
+	
+	USkeletalMeshComponent* ControlledMesh;
+	ControlledMesh = Cast<USkeletalMeshComponent>(ControlledSkeletalMesh.GetComponent(GetOwner()));
+
+	if(ControlledMesh)
+	{
+		ULiveLinkInstance* LiveLinkAnimInstance;
+		LiveLinkAnimInstance = Cast<ULiveLinkInstance>(ControlledMesh->GetAnimInstance());
+		if(LiveLinkAnimInstance)
+		{
+			LiveLinkAnimInstance->SetSubject(Subject);
+		}
+	}
+}
+
+FLiveLinkSubjectName UPerformerComponent::GetLiveLinkSubject() const
+{
+	return SubjectName;
+}
+
+void UPerformerComponent::SetEvaluateLiveLinkData(bool bEvaluateLinkLink)
+{
+	bEvaluateAnimation = bEvaluateLinkLink;
+	
+	USkeletalMeshComponent* ControlledMesh;
+	ControlledMesh = Cast<USkeletalMeshComponent>(ControlledSkeletalMesh.GetComponent(GetOwner()));
+
+	if(ControlledMesh)
+	{
+		ULiveLinkInstance* LiveLinkAnimInstance;
+		LiveLinkAnimInstance = Cast<ULiveLinkInstance>(ControlledMesh->GetAnimInstance());
+	
+		if (LiveLinkAnimInstance)
+		{
+			LiveLinkAnimInstance->EnableLiveLinkEvaluation(bEvaluateLinkLink);
+		}
+	}
+}
+
+bool UPerformerComponent::GetEvaluateLiveLinkData()
+{
+	return bEvaluateAnimation;
+}
+
+void UPerformerComponent::InitiateAnimation()
+{
+	TArray<USkeletalMeshComponent*> OwnerSkeletalMeshComponents;
+	GetOwner()->GetComponents(OwnerSkeletalMeshComponents);
+
+	//Set the mesh that will be controlled
+	USkeletalMeshComponent* ControlledMesh = Cast<USkeletalMeshComponent>(ControlledSkeletalMesh.GetComponent(GetOwner()));
+	
+	//If there is no controlled mesh, reset everything
+	if(!IsValid(ControlledMesh))
+	{
+		for (USkeletalMeshComponent* OwnerSkeletalMeshComponent : OwnerSkeletalMeshComponents)
+		{
+			OwnerSkeletalMeshComponent->SetAnimClass(nullptr);
+			OwnerSkeletalMeshComponent->InitAnim(true /*bForceReinit*/);
+			OwnerSkeletalMeshComponent->SetUpdateAnimationInEditor(true);
+		}
+	}
+
+	//Reset all skeletal meshes on owner
+	for (USkeletalMeshComponent* OwnerSkeletalMeshComponent : OwnerSkeletalMeshComponents)
+	{
+		//Null out the given leader pose to avoid a circular dependency.
+		OwnerSkeletalMeshComponent->SetLeaderPoseComponent(nullptr);
+		OwnerSkeletalMeshComponent->InitAnim(true /*bForceReinit*/);
+		
+		//Skip any skeletal meshes that already have anim instance set that's not a LiveLinkInstance
+		if(OwnerSkeletalMeshComponent->GetAnimClass()==ULiveLinkInstance::StaticClass())
+		{
+			OwnerSkeletalMeshComponent->SetAnimClass(nullptr);
+			OwnerSkeletalMeshComponent->InitAnim(true /*bForceReinit*/);
+			OwnerSkeletalMeshComponent->SetUpdateAnimationInEditor(true);
+		}
+	}
+	//Set the chosen mesh as leader
+	for (USkeletalMeshComponent* OwnerSkeletalMeshComponent : OwnerSkeletalMeshComponents)
+	{
+		if(IsValid(OwnerSkeletalMeshComponent))
+		{
+			if(OwnerSkeletalMeshComponent!=ControlledMesh && bForceOtherMeshesToFollowControlledMesh)
+			{
+				OwnerSkeletalMeshComponent->SetLeaderPoseComponent(ControlledMesh);
+			}
+		}
+	}
+
+	if(IsValid(ControlledMesh))
+	{
+		ControlledMesh->SetAnimClass(ULiveLinkInstance::StaticClass());
+		ControlledMesh->InitAnim(true /*bForceReinit*/);
+		ControlledMesh->SetUpdateAnimationInEditor(true);
+		ControlledMesh->bPropagateCurvesToFollowers = true;
+		ControlledMesh->VisibilityBasedAnimTickOption = EVisibilityBasedAnimTickOption::AlwaysTickPoseAndRefreshBones;
+		SetLiveLinkSubject(SubjectName);
+		SetEvaluateLiveLinkData(bEvaluateAnimation);
+	}	
+	bIsDirty=false;
+}
\ No newline at end of file
diff --git a/Engine/Plugins/Animation/PerformanceCaptureCore/Source/PerformanceCaptureCore/Private/RetargetAnimInstance.cpp b/Engine/Plugins/Animation/PerformanceCaptureCore/Source/PerformanceCaptureCore/Private/RetargetAnimInstance.cpp
new file mode 100644
--- /dev/null
+++ b/Engine/Plugins/Animation/PerformanceCaptureCore/Source/PerformanceCaptureCore/Private/RetargetAnimInstance.cpp
@@ -0,0 +1,38 @@
+// Copyright Epic Games, Inc. All Rights Reserved.
+
+#include "RetargetAnimInstance.h"
+#include "RetargetInstanceProxy.h"
+
+///////////////////////////////////
+/// Anim Instance 
+///////////////////////////////////
+
+void URetargetAnimInstance::UpdateCustomRetargetProfile(const FRetargetProfile &InRetargetProfile)
+{
+	FRetargetAnimInstanceProxy& Proxy = GetProxyOnGameThread<FRetargetAnimInstanceProxy>();
+	Proxy.RetargetNode->CustomRetargetProfile = InRetargetProfile;
+}
+
+FRetargetProfile URetargetAnimInstance::GetRetargetProfile()
+{
+	FRetargetAnimInstanceProxy& Proxy = GetProxyOnGameThread<FRetargetAnimInstanceProxy>();
+	FRetargetProfile OutProfile = Proxy.RetargetNode->CustomRetargetProfile;
+	return OutProfile;
+}
+
+void URetargetAnimInstance::NativeInitializeAnimation()
+{
+	FRetargetAnimInstanceProxy& Proxy = GetProxyOnGameThread<FRetargetAnimInstanceProxy>();
+	Proxy.Initialize(this);
+}
+
+void URetargetAnimInstance::ConfigureAnimInstance(UIKRetargeter* InIKRetargetAsset, TWeakObjectPtr<USkeletalMeshComponent> InSourceMeshComponent, FRetargetProfile InRetargetProfile)
+{
+	FRetargetAnimInstanceProxy& Proxy = GetProxyOnGameThread<FRetargetAnimInstanceProxy>();
+	Proxy.ConfigureAnimInstanceProxy(InIKRetargetAsset, InSourceMeshComponent, InRetargetProfile);
+}
+
+FAnimInstanceProxy* URetargetAnimInstance::CreateAnimInstanceProxy()
+{
+	return new FRetargetAnimInstanceProxy(this, &RetargetNode);
+}
diff --git a/Engine/Plugins/Animation/PerformanceCaptureCore/Source/PerformanceCaptureCore/Private/RetargetComponent.cpp b/Engine/Plugins/Animation/PerformanceCaptureCore/Source/PerformanceCaptureCore/Private/RetargetComponent.cpp
new file mode 100644
--- /dev/null
+++ b/Engine/Plugins/Animation/PerformanceCaptureCore/Source/PerformanceCaptureCore/Private/RetargetComponent.cpp
@@ -0,0 +1,244 @@
+// Copyright Epic Games, Inc. All Rights Reserved.
+
+#include "RetargetComponent.h"
+#include "LiveLinkInstance.h"
+#include "RetargetAnimInstance.h"
+
+// Sets default values for this component's properties
+URetargetComponent::URetargetComponent()
+{
+	// Set this component to be initialized when the game starts, and to be ticked every frame. 
+	PrimaryComponentTick.bCanEverTick = true;
+	PrimaryComponentTick.bStartWithTickEnabled = true;
+	bTickInEditor = true;
+}
+
+void URetargetComponent::SetCustomRetargetProfile(FRetargetProfile InProfile)
+{
+	CustomRetargetProfile = InProfile;
+
+	const TObjectPtr<USkeletalMeshComponent> ControlledMesh = Cast<USkeletalMeshComponent>(ControlledSkeletalMeshComponent.GetComponent(GetOwner()));
+	if(ControlledMesh)
+	{
+		if(ControlledMesh->GetAnimInstance())
+		{
+			URetargetAnimInstance* AnimInstance = Cast<URetargetAnimInstance>(ControlledMesh->GetAnimInstance());
+			AnimInstance->UpdateCustomRetargetProfile(CustomRetargetProfile);
+		}
+	}
+}
+
+FRetargetProfile URetargetComponent::GetCustomRetargetProfile()
+{
+	/// Check for valid components and refs and then get the custom retarget profile struct
+	FRetargetProfile OutProfile;
+    TObjectPtr<USkeletalMeshComponent> ControlledMesh;
+    ControlledMesh = Cast<USkeletalMeshComponent>(ControlledSkeletalMeshComponent.GetComponent(GetOwner()));
+	if(ControlledMesh)
+	{
+		if(ControlledMesh->GetAnimInstance())
+		{
+			URetargetAnimInstance* AnimInstance;
+			AnimInstance = Cast<URetargetAnimInstance>(ControlledMesh->GetAnimInstance());
+			OutProfile = AnimInstance->GetRetargetProfile();
+		}
+	}
+	return OutProfile;
+}
+
+void URetargetComponent::SetSourcePerformerMesh(USkeletalMeshComponent* InPerformerMesh)
+{
+	SourceSkeletalMeshComponent.OverrideComponent = InPerformerMesh;
+	InitiateAnimation();
+}
+
+
+void URetargetComponent::SetControlledMesh(USkeletalMeshComponent* InControlledMesh)
+{
+	ControlledSkeletalMeshComponent.OverrideComponent = InControlledMesh;
+	InitiateAnimation();
+}
+
+void URetargetComponent::SetRetargetAsset(UIKRetargeter* InRetargetAsset)
+{
+	RetargetAsset = InRetargetAsset; //It is possible to pass a nullptr to the AnimInstance 
+
+	/// Check for valid components and refs and then reint animation to reset pose
+	TObjectPtr<USkeletalMeshComponent> ControlledMesh = Cast<USkeletalMeshComponent>(ControlledSkeletalMeshComponent.GetComponent(GetOwner()));
+	if(ControlledMesh)
+	{
+		if(ControlledMesh->GetAnimInstance())
+		{
+			URetargetAnimInstance* AnimInstance = Cast<URetargetAnimInstance>(ControlledMesh->GetAnimInstance());
+
+			const TObjectPtr<USkeletalMeshComponent> SourceMesh = Cast<USkeletalMeshComponent>(SourceSkeletalMeshComponent.GetComponent(GetOwner()));
+		
+			if(SourceSkeletalMeshComponent.GetComponent(GetOwner()) && AnimInstance)
+			{
+				AnimInstance->ConfigureAnimInstance(RetargetAsset, SourceMesh, CustomRetargetProfile);
+				ControlledMesh->InitAnim(true /*bForceReinit*/);
+			}
+			ControlledMesh->InitAnim(true /*bForceReinit*/);
+		}
+	}
+}
+
+void URetargetComponent::OnRegister()
+{
+	Super::OnRegister();
+
+	const USkeletalMeshComponent* ControlledMesh = Cast<USkeletalMeshComponent> (ControlledSkeletalMeshComponent.GetComponent(GetOwner()));
+
+	const TObjectPtr<URetargetAnimInstance> AnimInstance = Cast<URetargetAnimInstance>(ControlledMesh->GetAnimInstance());
+		
+	if(AnimInstance) //Only set these properties if we have a valid AnimInstance of the correct class.
+	{
+		SetForceOtherMeshesToFollowControlledMesh(bForceOtherMeshesToFollowControlledMesh);
+		SetCustomRetargetProfile(CustomRetargetProfile);
+	}
+	
+	bIsDirty = true;
+}
+
+void URetargetComponent::TickComponent(float DeltaTime, ELevelTick TickType, FActorComponentTickFunction* ThisTickFunction)
+{
+	Super::TickComponent(DeltaTime, TickType, ThisTickFunction);
+	if(bIsDirty)
+	{
+		InitiateAnimation();
+	}
+}
+
+void URetargetComponent::DestroyComponent(bool bPromoteChildren)
+{
+	//When the component is removed, reset the state of the skelmeshes in the owner actor
+	TArray<USkeletalMeshComponent*> OwnerSkeletalMeshComponents;
+	GetOwner()->GetComponents(OwnerSkeletalMeshComponents);
+	
+	for (USkeletalMeshComponent* OwnerSkeletalMeshComponent : OwnerSkeletalMeshComponents)
+	{
+		//If the anim class is LiveLinkInstance assume user does not want to reset that skeletal mesh component
+		if(OwnerSkeletalMeshComponent->GetAnimClass()!=ULiveLinkInstance::StaticClass())
+		{
+			OwnerSkeletalMeshComponent->SetAnimClass(nullptr);
+			OwnerSkeletalMeshComponent->SetLeaderPoseComponent(nullptr);
+			OwnerSkeletalMeshComponent->InitAnim(true /*bForceReinit*/);
+			OwnerSkeletalMeshComponent->SetUpdateAnimationInEditor(false);
+		}
+	}
+	Super::DestroyComponent(bPromoteChildren);
+}
+
+#if WITH_EDITOR
+void URetargetComponent::PostEditChangeProperty(FPropertyChangedEvent& PropertyChangedEvent)
+{
+	Super::PostEditChangeProperty(PropertyChangedEvent);
+	
+	if(PropertyChangedEvent.Property != nullptr)
+	{
+		if (PropertyChangedEvent.Property->GetFName() == GET_MEMBER_NAME_CHECKED(URetargetComponent, SourceSkeletalMeshComponent))
+		{
+			bIsDirty=true;
+		}
+
+		Super::PostEditChangeProperty(PropertyChangedEvent);
+		if (PropertyChangedEvent.Property->GetFName() == GET_MEMBER_NAME_CHECKED(URetargetComponent, ControlledSkeletalMeshComponent))
+		{
+			bIsDirty=true;
+		}
+
+		Super::PostEditChangeProperty(PropertyChangedEvent);
+		if (PropertyChangedEvent.Property->GetFName() == GET_MEMBER_NAME_CHECKED(URetargetComponent, bForceOtherMeshesToFollowControlledMesh))
+		{
+			SetForceOtherMeshesToFollowControlledMesh(bForceOtherMeshesToFollowControlledMesh);
+		}
+
+		Super::PostEditChangeProperty(PropertyChangedEvent);
+		if (PropertyChangedEvent.Property->GetFName() == GET_MEMBER_NAME_CHECKED(URetargetComponent, RetargetAsset))
+		{
+			SetRetargetAsset(RetargetAsset);
+			bIsDirty=true;
+		}
+		
+		if (PropertyChangedEvent.Property->GetFName() == GET_MEMBER_NAME_CHECKED(URetargetComponent, CustomRetargetProfile))
+		{
+			SetCustomRetargetProfile(CustomRetargetProfile);
+		}
+	}
+}
+#endif
+
+void URetargetComponent::SetForceOtherMeshesToFollowControlledMesh(bool bInBool)
+{
+	bForceOtherMeshesToFollowControlledMesh = bInBool;
+	
+	USkeletalMeshComponent* ControlledMesh = Cast<USkeletalMeshComponent> (ControlledSkeletalMeshComponent.GetComponent
+(GetOwner()));
+	const USkeletalMeshComponent* SourceMesh = Cast<USkeletalMeshComponent>(SourceSkeletalMeshComponent.GetComponent(GetOwner()));
+	
+	TArray<USkeletalMeshComponent*> OwnerSkeletalMeshComponents;
+	GetOwner()->GetComponents(OwnerSkeletalMeshComponents);
+
+	if(IsValid(ControlledMesh) && IsValid(SourceMesh))
+	{
+		for (USkeletalMeshComponent* OwnerSkeletalMeshComponent : OwnerSkeletalMeshComponents)
+		{
+			if(IsValid(OwnerSkeletalMeshComponent))
+			{
+				const bool bShouldResetMesh = OwnerSkeletalMeshComponent->GetAnimClass()!=ULiveLinkInstance::StaticClass()
+				&& OwnerSkeletalMeshComponent!=ControlledMesh && OwnerSkeletalMeshComponent!=SourceMesh;
+				if(bShouldResetMesh)
+				{
+					OwnerSkeletalMeshComponent->SetAnimClass(nullptr);
+					OwnerSkeletalMeshComponent->InitAnim(true /*bForceReinit*/);
+					OwnerSkeletalMeshComponent->SetUpdateAnimationInEditor(true);
+					OwnerSkeletalMeshComponent->SetLeaderPoseComponent(nullptr);
+				}	
+			}
+		}
+		
+		if(bForceOtherMeshesToFollowControlledMesh)
+		{
+			//Set all skeletal meshes to tick in editor
+			for (USkeletalMeshComponent* OwnerSkeletalMeshComponent : OwnerSkeletalMeshComponents)
+			{
+				OwnerSkeletalMeshComponent->SetUpdateAnimationInEditor(bInBool);
+
+				///If not Controlled or the Source mesh then set skeletal meshes to follow the ControlledMesh. Exception for LiveLink instances as we assume those are being driven directly be mocap data using the PCapPerformerComponent
+				///
+				const bool bShouldSetLeaderPose  = OwnerSkeletalMeshComponent!=ControlledMesh && OwnerSkeletalMeshComponent->GetAnimClass()!=ULiveLinkInstance::StaticClass() && OwnerSkeletalMeshComponent!=SourceMesh;
+				if(bShouldSetLeaderPose)
+				{
+					OwnerSkeletalMeshComponent->SetLeaderPoseComponent(ControlledMesh,true /*ForceUpdate*/, true /*FollowerShouldTickPose*/);
+				}
+			}
+		}
+	}
+}
+
+void URetargetComponent::InitiateAnimation()
+{
+	USkeletalMeshComponent* ControlledMesh = Cast<USkeletalMeshComponent> (ControlledSkeletalMeshComponent.GetComponent(GetOwner()));
+	USkeletalMeshComponent* SourceMesh = Cast<USkeletalMeshComponent>(SourceSkeletalMeshComponent.GetComponent(GetOwner()));
+	
+	const bool bShouldBeReinitialized = IsValid(SourceMesh) && IsValid(RetargetAsset) && IsValid(ControlledMesh);
+	
+	if(bShouldBeReinitialized)
+	{
+		//Set the anim instance class on the controlled mesh to use RetargetAnimInstance
+		ControlledMesh->SetAnimClass(URetargetAnimInstance::StaticClass());
+
+		TObjectPtr<URetargetAnimInstance> AnimInstance = Cast<URetargetAnimInstance>(ControlledMesh->GetAnimInstance());
+		
+		if(AnimInstance)
+		{
+			AnimInstance->ConfigureAnimInstance(RetargetAsset, SourceMesh, CustomRetargetProfile);
+			ControlledMesh->SetUpdateAnimationInEditor(true);
+			ControlledMesh->bPropagateCurvesToFollowers = true;
+			ControlledMesh->VisibilityBasedAnimTickOption = EVisibilityBasedAnimTickOption::AlwaysTickPoseAndRefreshBones;
+			ControlledMesh->InitAnim(true /*bForceReinit*/);
+		}
+		ControlledMesh->InitAnim(true /*bForceReinit*/);
+	}
+	bIsDirty=false;
+}
\ No newline at end of file
diff --git a/Engine/Plugins/Animation/PerformanceCaptureCore/Source/PerformanceCaptureCore/Private/RetargetInstanceProxy.cpp b/Engine/Plugins/Animation/PerformanceCaptureCore/Source/PerformanceCaptureCore/Private/RetargetInstanceProxy.cpp
new file mode 100644
--- /dev/null
+++ b/Engine/Plugins/Animation/PerformanceCaptureCore/Source/PerformanceCaptureCore/Private/RetargetInstanceProxy.cpp
@@ -0,0 +1,81 @@
+// Copyright Epic Games, Inc. All Rights Reserved.
+#include "RetargetInstanceProxy.h"
+
+#include UE_INLINE_GENERATED_CPP_BY_NAME(RetargetInstanceProxy)
+
+
+/////////////////////////////////
+/// Anim instance proxy struct
+/////////////////////////////////
+
+FRetargetAnimInstanceProxy::FRetargetAnimInstanceProxy(UAnimInstance* InAnimInstance, FAnimNode_RetargetPoseFromMesh* InRetargetNode)
+	: FAnimInstanceProxy(InAnimInstance),
+	RetargetNode(InRetargetNode)
+{
+}
+
+void FRetargetAnimInstanceProxy::Initialize(UAnimInstance* InAnimInstance)
+{
+	FAnimInstanceProxy::Initialize(InAnimInstance);
+
+	//Initialize instance manually
+	FAnimationInitializeContext InitContext(this);
+	RetargetNode->Initialize_AnyThread(InitContext);
+}
+
+void FRetargetAnimInstanceProxy::PreUpdate(UAnimInstance* InAnimInstance, float DeltaSeconds)
+{
+	if(InAnimInstance)
+	{
+		Super::PreUpdate(InAnimInstance, DeltaSeconds);
+				
+		if (RetargetNode->HasPreUpdate())
+		{
+			RetargetNode->PreUpdate(InAnimInstance);
+		}
+	}
+}
+
+// Cache bones if they are marked as invalid//
+void FRetargetAnimInstanceProxy::CacheBones()
+{
+	if(bBoneCachesInvalidated)
+	{
+		FAnimationCacheBonesContext Context(this);
+		RetargetNode->CacheBones_AnyThread(Context);
+		bBoneCachesInvalidated = false;
+	}
+}
+
+// Evaluate pose//
+bool FRetargetAnimInstanceProxy::Evaluate(FPoseContext& Output)
+{
+	Super::Evaluate(Output);
+	
+	/* This evaluates UAnimNode_RetargetPoseFromMesh */
+	RetargetNode->Evaluate_AnyThread(Output);
+	return true;
+}
+
+//Update animation//
+void FRetargetAnimInstanceProxy::UpdateAnimationNode(const FAnimationUpdateContext& InContext)
+{
+	Super::UpdateAnimationNode(InContext);
+
+	UpdateCounter.Increment();
+	
+	RetargetNode->Update_AnyThread(InContext);
+}
+
+void FRetargetAnimInstanceProxy::ConfigureAnimInstanceProxy(UIKRetargeter* InIKRetargetAsset, TWeakObjectPtr<USkeletalMeshComponent> InSourceMeshComponent, FRetargetProfile &InRetargetProfile)
+{
+	RetargetNode->IKRetargeterAsset = InIKRetargetAsset;
+	RetargetNode->bUseAttachedParent = false;
+	RetargetNode->SourceMeshComponent = InSourceMeshComponent;
+	RetargetNode->CustomRetargetProfile = InRetargetProfile;
+	
+	if (UIKRetargetProcessor* Processor = RetargetNode->GetRetargetProcessor())
+	{
+		Processor->SetNeedsInitialized();
+	}
+}
diff --git a/Engine/Plugins/Animation/PerformanceCaptureCore/Source/PerformanceCaptureCore/Private/RetargetInstanceProxy.h b/Engine/Plugins/Animation/PerformanceCaptureCore/Source/PerformanceCaptureCore/Private/RetargetInstanceProxy.h
new file mode 100644
--- /dev/null
+++ b/Engine/Plugins/Animation/PerformanceCaptureCore/Source/PerformanceCaptureCore/Private/RetargetInstanceProxy.h
@@ -0,0 +1,43 @@
+// Copyright Epic Games, Inc. All Rights Reserved.
+#pragma once
+
+#include "RetargetAnimInstance.h"
+#include "Animation/AnimInstanceProxy.h"
+
+#include "RetargetInstanceProxy.generated.h"
+
+class UIKRetargeter;
+class USkeletalMeshComponent;
+struct FAnimNode_RetargetPoseFromMesh;
+
+////////////////////////////////////
+/// Anim Instance Proxy
+////////////////////////////////////
+USTRUCT()
+struct FRetargetAnimInstanceProxy : public FAnimInstanceProxy
+{
+public:
+	
+	GENERATED_BODY()
+
+	FRetargetAnimInstanceProxy() = default; //Constructor
+	
+	FRetargetAnimInstanceProxy(UAnimInstance* InAnimInstance, FAnimNode_RetargetPoseFromMesh* InRetargetNode);
+	
+public:	
+	/* FAnimInstanceProxy Instance */ 
+	virtual void Initialize(UAnimInstance* InAnimInstance) override; 
+	virtual void PreUpdate(UAnimInstance* InAnimInstance, float DeltaSeconds) override;
+	virtual void CacheBones() override;
+	virtual bool Evaluate(FPoseContext& Output) override;
+	virtual void UpdateAnimationNode(const FAnimationUpdateContext& InContext) override;
+	/* END FAnimInstanceProxy Instance */
+	
+	void ConfigureAnimInstanceProxy(UIKRetargeter* InIKRetargetAsset, TWeakObjectPtr<USkeletalMeshComponent> InSourceMeshComponent,FRetargetProfile &InRetargetProfile);
+
+	FAnimNode_RetargetPoseFromMesh* RetargetNode = nullptr;
+
+private:	
+	FRetargetProfile* RetargetProfile = nullptr;
+	
+};
diff --git a/Engine/Plugins/Animation/PerformanceCaptureCore/Source/PerformanceCaptureCore/Public/CaptureCharacter.h b/Engine/Plugins/Animation/PerformanceCaptureCore/Source/PerformanceCaptureCore/Public/CaptureCharacter.h
new file mode 100644
--- /dev/null
+++ b/Engine/Plugins/Animation/PerformanceCaptureCore/Source/PerformanceCaptureCore/Public/CaptureCharacter.h
@@ -0,0 +1,86 @@
+// Copyright Epic Games, Inc. All Rights Reserved.
+#pragma once
+
+#include "CoreMinimal.h"
+#include "Animation/SkeletalMeshActor.h"
+#include "CapturePerformer.h"
+#include "Retargeter/IKRetargetProfile.h"
+
+#include "CaptureCharacter.generated.h"
+
+class URetargetComponent;
+class UIKRetargeter;
+
+UCLASS(Blueprintable, ClassGroup ="Performance Capture", Category = "Performance Capture",  HideCategories = ("Mesh", "Rendering", "Animation", "LOD", "Misc", "Physics", "Streaming"))
+class ACaptureCharacter : public ASkeletalMeshActor
+{
+	GENERATED_BODY()
+
+public:
+	// Sets default values for this actor's properties
+	ACaptureCharacter();
+
+	/**
+	* CapturePerformer Actor that will be the source for retargeting.
+	*/
+	UPROPERTY(BlueprintReadOnly, EditAnywhere, Category = "Performance Capture")
+	TObjectPtr<ACapturePerformer> SourcePerformer;
+
+	/**
+	* The IKRetarget Asset to use for retargeting between the SourcePerformer and this Character.
+	*/
+	UPROPERTY(BlueprintReadOnly, EditAnywhere, Category = "Performance Capture")
+	TObjectPtr<UIKRetargeter> RetargetAsset;
+
+	/**
+	* Force all skeletal meshes to use the root SkeletalMesh as their Leader. Default = True.
+	*/
+	UPROPERTY(BlueprintReadOnly, EditAnywhere, Category = "Performance Capture")
+	bool bForceAllSkeletalMeshesToFollowLeader = true;
+	
+	/**
+	* Set a Custom Retarget Profile.
+	* @param InProfile New Retarget Profile.
+	*/
+	UFUNCTION(BlueprintCallable, Category = "Performance Capture")
+	void SetCustomRetargetProfile(FRetargetProfile InProfile);
+
+	/**
+	* Get a Custom Retarget Profile.
+	* @return Current Custom Retarget Profile.
+	*/
+	UFUNCTION(BlueprintCallable, Category = "Performance Capture")
+	FRetargetProfile GetCustomRetargetProfile();
+
+	/**
+	* Set the Source CapturePerformer Actor.
+	* @param InPerformer New Performer.
+	*/
+	UFUNCTION(BlueprintCallable, Category = "Performance Capture")
+	void SetSourcePerformer(ACapturePerformer* InPerformer);
+
+	/**
+	* Set the Retarget Asset.
+	* @param InRetargetAsset New IKRetarget Asset.
+	*/
+	UFUNCTION(BlueprintCallable, Category = "Performance Capture")
+	void SetRetargetAsset(UIKRetargeter* InRetargetAsset);
+
+	/**
+	* Force all Skeletal Meshes to follow the Controlled Skeletal mesh. 
+	* @param InFollowLeader New Bool.
+	*/
+	UFUNCTION(BlueprintCallable, Category = "Performance Capture")
+	void SetForceAllSkeletalMeshesToFollowLeader (bool InFollowLeader);
+
+protected:
+	virtual void PostRegisterAllComponents() override;
+#if WITH_EDITOR
+	virtual void PostEditChangeProperty(FPropertyChangedEvent& PropertyChangedEvent) override;
+#endif
+
+private:
+	// retarget component is private and hidden from the editor UI
+	UPROPERTY()
+	TObjectPtr<URetargetComponent> RetargetComponent;
+};
diff --git a/Engine/Plugins/Animation/PerformanceCaptureCore/Source/PerformanceCaptureCore/Public/CapturePerformer.h b/Engine/Plugins/Animation/PerformanceCaptureCore/Source/PerformanceCaptureCore/Public/CapturePerformer.h
new file mode 100644
--- /dev/null
+++ b/Engine/Plugins/Animation/PerformanceCaptureCore/Source/PerformanceCaptureCore/Public/CapturePerformer.h
@@ -0,0 +1,89 @@
+// Copyright Epic Games, Inc. All Rights Reserved.
+
+#pragma once
+
+#include "CoreMinimal.h"
+#include "Animation/SkeletalMeshActor.h"
+#include "LiveLinkTypes.h"
+#include "LiveLinkInstance.h"
+#include "CapturePerformer.generated.h"
+
+class ULiveLinkInstance;
+class UPerformerComponent;
+
+UCLASS(Blueprintable, Category="Performance Capture")
+class ACapturePerformer : public ASkeletalMeshActor
+{
+	GENERATED_BODY()
+
+public:
+	// Sets default values for this actor's properties
+	ACapturePerformer();
+
+	/**
+	* LiveLink Subject Name. Must have the Animation Role Type.
+	*/
+	UPROPERTY(BlueprintReadOnly, EditAnywhere, Category = "Performance Capture")
+	FLiveLinkSubjectName SubjectName;
+
+	/**
+	* Force all skeletal meshes to follow the pose of the Controlled Skeletal Mesh.
+	*/
+	UPROPERTY(BlueprintReadOnly, EditAnywhere, Category = "Performance Capture")
+	bool bForceAllMeshesToFollowLeader = true;
+
+	/**
+	* Evaluate LiveLink animation. Set to false to pause animation.
+	*/
+	UPROPERTY(BlueprintReadOnly, EditAnywhere, Category = "Performance Capture")
+	bool bEvaluateAnimation = true;
+
+	/**
+	* Set the LiveLink Subject Name. Subject must have the Animation Role Type.
+	* @param Subject New LiveLink Subject.
+	*/
+	UFUNCTION(BlueprintCallable, Category = "Performance Capture")
+	void SetLiveLinkSubject(FLiveLinkSubjectName Subject);
+
+	/**
+	* Get the LiveLink Subject Name.
+	* @return FLiveLinkSubjectName Current LiveLink Subject.
+	*/
+	UFUNCTION(BlueprintCallable, Category ="Performance Capture")
+	FLiveLinkSubjectName GetLiveLinkSubject() const;
+
+	/**
+	* Set the LiveLink data to update the Skeletal Mesh pose.
+	* @param bEvaluateLinkLink Drive or pause the Skeletal Mesh from LiveLink Subject data.
+	*/
+	UFUNCTION(BlueprintCallable, Category ="Performance Capture")
+	void SetEvaluateLiveLinkData(bool bEvaluateLinkLink);
+
+	/**
+	* Get state of LiveLink data evaluation. True = animation is updated.
+	* @return bool Is LiveLink data being evaluated.
+	*/
+	UFUNCTION(BlueprintPure, Category ="Performance Capture")
+	bool GetEvaluateLiveLinkData();
+
+	/**
+	* Set the Skeletal Mesh Asset of the root Skeletal Mesh.
+	* @param MocapMesh New Skeletal Mesh Asset.
+	*/
+	UFUNCTION(BlueprintCallable, Category = "Performance Capture")
+	void SetMocapMesh(USkeletalMesh* MocapMesh);
+
+protected:
+	// Called when the game starts or when spawned
+	virtual void BeginPlay() override;
+	virtual void PostInitializeComponents() override;
+	virtual void PostRegisterAllComponents() override;
+	virtual void OnConstruction(const FTransform& Transform) override;
+
+#if WITH_EDITOR
+	virtual void PostEditChangeProperty(FPropertyChangedEvent& PropertyChangedEvent) override;
+#endif
+	
+private:
+	TObjectPtr<UPerformerComponent> PerformerComponent;
+};
\ No newline at end of file
diff --git a/Engine/Plugins/Animation/PerformanceCaptureCore/Source/PerformanceCaptureCore/Public/PerformanceCaptureCore.h b/Engine/Plugins/Animation/PerformanceCaptureCore/Source/PerformanceCaptureCore/Public/PerformanceCaptureCore.h
new file mode 100644
--- /dev/null
+++ b/Engine/Plugins/Animation/PerformanceCaptureCore/Source/PerformanceCaptureCore/Public/PerformanceCaptureCore.h
@@ -0,0 +1,15 @@
+﻿// Copyright Epic Games, Inc. All Rights Reserved.
+
+#pragma once
+
+#include "CoreMinimal.h"
+#include "Modules/ModuleManager.h"
+
+class FPerformanceCaptureCoreModule : public IModuleInterface
+{
+public:
+
+	/** IModuleInterface implementation */
+	virtual void StartupModule() override;
+	virtual void ShutdownModule() override;
+};
diff --git a/Engine/Plugins/Animation/PerformanceCaptureCore/Source/PerformanceCaptureCore/Public/PerformerComponent.h b/Engine/Plugins/Animation/PerformanceCaptureCore/Source/PerformanceCaptureCore/Public/PerformerComponent.h
new file mode 100644
--- /dev/null
+++ b/Engine/Plugins/Animation/PerformanceCaptureCore/Source/PerformanceCaptureCore/Public/PerformerComponent.h
@@ -0,0 +1,91 @@
+// Copyright Epic Games, Inc. All Rights Reserved.
+
+#pragma once
+
+#include "CoreMinimal.h"
+#include "LiveLinkTypes.h"
+#include "Components/ActorComponent.h"
+
+#include "PerformerComponent.generated.h"
+
+class ULiveLinkInstance;
+
+UCLASS(BlueprintType, DisplayName="Performer Component",  ClassGroup=("Performance Capture"), meta=(BlueprintSpawnableComponent))
+class UPerformerComponent : public UActorComponent
+{
+	GENERATED_BODY()
+
+public:
+	// Sets default values for this component's properties
+	UPerformerComponent();
+
+	/**
+	* LiveLink Subject Name. Must have the Animation Role Type.
+	*/
+	UPROPERTY(BlueprintReadWrite, EditAnywhere, Category = "Performance Capture")
+	FLiveLinkSubjectName SubjectName;
+
+	/**
+	* The Skeletal Mesh driven by the LiveLink subject. Skeleton mush be compatible with the LiveLink subject's bone hierarchy. 
+	*/
+	UPROPERTY(BlueprintReadWrite, EditInstanceOnly, Category = "Performance Capture", meta=(UseComponentPicker, AllowedClasses = "/Script/Engine.SkeletalMeshComponent"))
+	FComponentReference ControlledSkeletalMesh;
+	
+	/**
+	* Evaluate LiveLink animation. Set to false to pause animation.
+	*/
+	UPROPERTY(BlueprintReadWrite, EditAnywhere, Category = "Performance Capture")
+	bool bEvaluateAnimation = true;
+	
+	/**
+	* Force all other skeletal meshes in the Owner Actor to follow the pose of the Controlled Skeletal Mesh.
+	*/
+	UPROPERTY(BlueprintReadWrite, EditAnywhere, Category = "Performance Capture")
+	bool bForceOtherMeshesToFollowControlledMesh = true;
+
+	/**
+	* Set the LiveLink Subject Name. Subject must have the Animation Role Type.
+	* @param Subject New LiveLink Subject.
+	*/
+	UFUNCTION(BlueprintCallable, Category = "Performance Capture")
+	void SetLiveLinkSubject(FLiveLinkSubjectName Subject);
+	
+	/**
+	* Get the LiveLink Subject Name.
+	* @return FLiveLinkSubjectName Current LiveLink Subject.
+	*/
+	UFUNCTION(BlueprintCallable, Category ="Performance Capture")
+	FLiveLinkSubjectName GetLiveLinkSubject() const;
+	
+	/**
+	* Set the LiveLink data to update the Skeletal Mesh pose.
+	* @param bEvaluateLinkLink Drive or pause the Skeletal Mesh from LiveLink Subject data.
+	*/
+	UFUNCTION(BlueprintCallable, Category ="Performance Capture")
+	void SetEvaluateLiveLinkData(bool bEvaluateLinkLink);
+	
+	/**
+	* Get the LiveLink Subject Name. Subject must have the Animation Role Type.
+	* @return bool Is LiveLink data being evaluated.
+	*/
+	UFUNCTION(BlueprintPure, Category ="Performance Capture")
+	bool GetEvaluateLiveLinkData();
+
+protected:
+	virtual void DestroyComponent(bool bPromoteChildren) override;
+
+	virtual void OnRegister() override;
+#if WITH_EDITOR
+	virtual void PostEditChangeProperty(FPropertyChangedEvent& PropertyChangedEvent) override;
+#endif
+
+public:
+	// Called every frame
+	virtual void TickComponent(float DeltaTime, ELevelTick TickType, FActorComponentTickFunction* ThisTickFunction) override;
+	void InitiateAnimation();
+
+private:
+
+	bool bIsDirty;
+	
+};
diff --git a/Engine/Plugins/Animation/PerformanceCaptureCore/Source/PerformanceCaptureCore/Public/RetargetAnimInstance.h b/Engine/Plugins/Animation/PerformanceCaptureCore/Source/PerformanceCaptureCore/Public/RetargetAnimInstance.h
new file mode 100644
--- /dev/null
+++ b/Engine/Plugins/Animation/PerformanceCaptureCore/Source/PerformanceCaptureCore/Public/RetargetAnimInstance.h
@@ -0,0 +1,56 @@
+// Copyright Epic Games, Inc. All Rights Reserved.
+
+#pragma once
+
+#include "Animation/AnimInstance.h"
+#include "Animation/AnimNodeBase.h"
+#include "AnimNodes/AnimNode_RetargetPoseFromMesh.h"
+
+#include "RetargetAnimInstance.generated.h"
+
+class UIKRetargeter;
+class USkeletalMeshComponent;
+struct FAnimNode_RetargetPoseFromMesh;
+struct FRetargetAnimInstanceProxy;
+
+///////////////////////////
+/// Anim Instance 
+///////////////////////////
+
+UCLASS(Transient, NotBlueprintable, BlueprintType, MinimalAPI)
+class URetargetAnimInstance : public UAnimInstance
+{
+	GENERATED_BODY()
+
+	friend FRetargetAnimInstanceProxy;
+
+public:
+	/**
+	* Configure Retarget AnimInstance
+	* @param InIKRetargetAsset IKRetargeter Asset to use for retargeting.
+	* @param InSourceMeshComponent The mesh component that is the source of motion.
+	* @param InRetargetProfile Retarget Profile struct to use.
+	*/
+	void ConfigureAnimInstance(UIKRetargeter* InIKRetargetAsset,TWeakObjectPtr<USkeletalMeshComponent> InSourceMeshComponent, FRetargetProfile InRetargetProfile);
+	
+	/**
+	* Update the RetargetProfile.
+	* @param InRetargetProfile Retarget Profile struct to use.
+	*/
+	void UpdateCustomRetargetProfile(const FRetargetProfile &InRetargetProfile);
+
+	/**
+	* Retrieve the current RetargetProfile struct.
+	* @return  Retarget Profile struct to use.
+	*/
+	FRetargetProfile GetRetargetProfile();
+
+protected:
+	/** UAnimInstance interface */
+	virtual void NativeInitializeAnimation() override;
+	virtual FAnimInstanceProxy* CreateAnimInstanceProxy() override;
+	/** UAnimInstance interface end*/
+	
+	UPROPERTY()
+	FAnimNode_RetargetPoseFromMesh RetargetNode;
+};
diff --git a/Engine/Plugins/Animation/PerformanceCaptureCore/Source/PerformanceCaptureCore/Public/RetargetComponent.h b/Engine/Plugins/Animation/PerformanceCaptureCore/Source/PerformanceCaptureCore/Public/RetargetComponent.h
new file mode 100644
--- /dev/null
+++ b/Engine/Plugins/Animation/PerformanceCaptureCore/Source/PerformanceCaptureCore/Public/RetargetComponent.h
@@ -0,0 +1,106 @@
+// Copyright Epic Games, Inc. All Rights Reserved.
+
+#pragma once
+
+#include "CoreMinimal.h"
+#include "Components/ActorComponent.h"
+#include "Retargeter/IKRetargetProfile.h"
+
+#include "RetargetComponent.generated.h"
+
+
+class UIKRetargeter;
+
+UCLASS(BlueprintType, ClassGroup=("Performance Capture"), meta=(BlueprintSpawnableComponent), DisplayName = "Retarget Component")
+class URetargetComponent : public UActorComponent
+{
+	GENERATED_BODY()
+
+public:
+	// Sets default values for this component's properties
+	URetargetComponent();
+
+	/**
+	* Skeletal Mesh component that will be the source for retargeting. Can be on the Owner Actor or another Actor in the same level.
+	*/
+	UPROPERTY(BlueprintReadWrite, EditInstanceOnly, Category = "Performance Capture",  meta = (UseComponentPicker, AllowAnyActor, AllowedClasses = "/Script/Engine.SkeletalMeshComponent"))
+	FComponentReference SourceSkeletalMeshComponent ;
+
+	/**
+	* Skeletal Mesh that will be driven by the IKRetargeter. Limited to skeletal meshes on this component's Owner Actor.
+	*/
+	UPROPERTY(BlueprintReadWrite, EditAnywhere,Category = "Performance Capture", meta = (UseComponentPicker, AllowedClasses = "/Script/Engine.SkeletalMeshComponent"))
+	FComponentReference ControlledSkeletalMeshComponent;
+
+	/**
+	* Force all skeletal meshes to use the ControlledSkeletalMeshComponent as their Leader. Default = True.
+	*/
+	UPROPERTY(BlueprintReadWrite, EditAnywhere, Category = "Performance Capture")
+	bool bForceOtherMeshesToFollowControlledMesh = true;
+
+	/**
+	* The IKRetarget Asset to use for retargeting between the source and controlled skeletal meshes.
+	*/
+	UPROPERTY(BlueprintReadOnly, EditAnywhere, Category = "Performance Capture")
+	TObjectPtr<UIKRetargeter> RetargetAsset;
+
+	/**
+	* Custom Retarget Profile. Should be used to override Retarget settings on the RetargetAsset.
+	*/
+	UPROPERTY(BlueprintReadWrite, Interp, EditAnywhere, Category = "Performance Capture")
+	FRetargetProfile CustomRetargetProfile;
+	
+	/**
+	* Set the Source Performer Mesh.
+	* @param InPerformerMesh New Source Skeletal Mesh.
+	*/
+	UFUNCTION(BlueprintCallable, Category = "Performance Capture")
+	void SetSourcePerformerMesh(USkeletalMeshComponent* InPerformerMesh);
+	
+	/**
+	* Set the Controlled Skeletal Mesh.
+	* @param InControlledMesh New Controlled Skeletal Mesh Component.
+	*/
+	UFUNCTION(BlueprintCallable, Category = "Performance Capture")
+	void SetControlledMesh(USkeletalMeshComponent* InControlledMesh);
+
+	/**
+	* Set the Retarget Asset.
+	* @param InRetargetAsset New IKRetarget Asset.
+	*/
+	UFUNCTION(BlueprintCallable, Category = "Performance Capture")
+	void SetRetargetAsset(UIKRetargeter* InRetargetAsset);
+
+	/**
+	* Set a Custom Retarget Profile.
+	* @param InProfile New Retarget Profile.
+	*/
+	UFUNCTION(BlueprintCallable, Category = "Performance Capture")
+	void SetCustomRetargetProfile(FRetargetProfile InProfile);
+
+	/**
+	* Get Retarget Profile.
+	* @return Current Custom Retarget Profile.
+	*/
+	UFUNCTION(BlueprintCallable, Category = "Performance Capture")
+	FRetargetProfile GetCustomRetargetProfile();
+	
+protected:
+	// Called when the game starts
+	virtual void OnRegister() override; //This is the equivalent of OnConstruction in an actor
+	virtual void TickComponent(float DeltaTime, ELevelTick TickType, FActorComponentTickFunction* ThisTickFunction) override;
+	virtual void DestroyComponent(bool bPromoteChildren) override;
+
+#if WITH_EDITOR
+	virtual void PostEditChangeProperty(FPropertyChangedEvent& PropertyChangedEvent) override;
+#endif
+	
+private:
+	bool bIsDirty = true;
+public:
+	
+	void InitiateAnimation();
+
+	void SetForceOtherMeshesToFollowControlledMesh(bool bInBool);
+};
+
