Enable Live Link Component and Performance Capture Core

diff --git a/Engine/Plugins/Animation/PerformanceCaptureCore/Source/PerformanceCaptureCore/Private/RetargetComponent.cpp b/Engine/Plugins/Animation/PerformanceCaptureCore/Source/PerformanceCaptureCore/Private/RetargetComponent.cpp
--- a/Engine/Plugins/Animation/PerformanceCaptureCore/Source/PerformanceCaptureCore/Private/RetargetComponent.cpp
+++ b/Engine/Plugins/Animation/PerformanceCaptureCore/Source/PerformanceCaptureCore/Private/RetargetComponent.cpp
@@ -89,14 +89,17 @@ void URetargetComponent::OnRegister()
 
 	const USkeletalMeshComponent* ControlledMesh = Cast<USkeletalMeshComponent> (ControlledSkeletalMeshComponent.GetComponent(GetOwner()));
 
-	const TObjectPtr<URetargetAnimInstance> AnimInstance = Cast<URetargetAnimInstance>(ControlledMesh->GetAnimInstance());
-		
-	if(AnimInstance) //Only set these properties if we have a valid AnimInstance of the correct class.
+	if (ControlledMesh)
 	{
-		SetForceOtherMeshesToFollowControlledMesh(bForceOtherMeshesToFollowControlledMesh);
-		SetCustomRetargetProfile(CustomRetargetProfile);
+		const TObjectPtr<URetargetAnimInstance> AnimInstance = Cast<URetargetAnimInstance>(ControlledMesh->GetAnimInstance());
+
+		if(AnimInstance) //Only set these properties if we have a valid AnimInstance of the correct class.
+		{
+			SetForceOtherMeshesToFollowControlledMesh(bForceOtherMeshesToFollowControlledMesh);
+			SetCustomRetargetProfile(CustomRetargetProfile);
+		}
 	}
-	
+
 	bIsDirty = true;
 }
 
@@ -241,4 +244,4 @@ void URetargetComponent::InitiateAnimation()
 		ControlledMesh->InitAnim(true /*bForceReinit*/);
 	}
 	bIsDirty=false;
-}
\ No newline at end of file
+}
