Performance Capture: fix regression fallout from 32584701. Added call to increment version of the retargeter asset each time the animation instance is reinitialized.

diff --git a/Engine/Plugins/Animation/PerformanceCaptureCore/Source/PerformanceCaptureCore/Private/RetargetComponent.cpp b/Engine/Plugins/Animation/PerformanceCaptureCore/Source/PerformanceCaptureCore/Private/RetargetComponent.cpp
--- a/Engine/Plugins/Animation/PerformanceCaptureCore/Source/PerformanceCaptureCore/Private/RetargetComponent.cpp
+++ b/Engine/Plugins/Animation/PerformanceCaptureCore/Source/PerformanceCaptureCore/Private/RetargetComponent.cpp
@@ -235,6 +235,7 @@ void URetargetComponent::InitiateAnimation()
 		
 		if(AnimInstance)
 		{
+			RetargetAsset->IncrementVersion();
 			AnimInstance->ConfigureAnimInstance(RetargetAsset, SourceMesh, CustomRetargetProfile);
 			ControlledMesh->SetUpdateAnimationInEditor(true);
 			ControlledMesh->bPropagateCurvesToFollowers = true;
