Performance Capture Core:
Added DLL export macro to public classes - Performer Actor , Retarget Actor, RetargetComponent & PerformerComponent

diff --git a/Engine/Plugins/Animation/PerformanceCaptureCore/Source/PerformanceCaptureCore/Public/CaptureCharacter.h b/Engine/Plugins/Animation/PerformanceCaptureCore/Source/PerformanceCaptureCore/Public/CaptureCharacter.h
--- a/Engine/Plugins/Animation/PerformanceCaptureCore/Source/PerformanceCaptureCore/Public/CaptureCharacter.h
+++ b/Engine/Plugins/Animation/PerformanceCaptureCore/Source/PerformanceCaptureCore/Public/CaptureCharacter.h
@@ -12,7 +12,7 @@ class URetargetComponent;
 class UIKRetargeter;
 
 UCLASS(Blueprintable, ClassGroup ="Performance Capture", Category = "Performance Capture",  HideCategories = ("Mesh", "Rendering", "Animation", "LOD", "Misc", "Physics", "Streaming"))
-class ACaptureCharacter : public ASkeletalMeshActor
+class PERFORMANCECAPTURECORE_API ACaptureCharacter : public ASkeletalMeshActor
 {
 	GENERATED_BODY()
 
diff --git a/Engine/Plugins/Animation/PerformanceCaptureCore/Source/PerformanceCaptureCore/Public/CapturePerformer.h b/Engine/Plugins/Animation/PerformanceCaptureCore/Source/PerformanceCaptureCore/Public/CapturePerformer.h
--- a/Engine/Plugins/Animation/PerformanceCaptureCore/Source/PerformanceCaptureCore/Public/CapturePerformer.h
+++ b/Engine/Plugins/Animation/PerformanceCaptureCore/Source/PerformanceCaptureCore/Public/CapturePerformer.h
@@ -12,7 +12,7 @@ class ULiveLinkInstance;
 class UPerformerComponent;
 
 UCLASS(Blueprintable, Category="Performance Capture")
-class ACapturePerformer : public ASkeletalMeshActor
+class PERFORMANCECAPTURECORE_API ACapturePerformer : public ASkeletalMeshActor
 {
 	GENERATED_BODY()
 
diff --git a/Engine/Plugins/Animation/PerformanceCaptureCore/Source/PerformanceCaptureCore/Public/PerformerComponent.h b/Engine/Plugins/Animation/PerformanceCaptureCore/Source/PerformanceCaptureCore/Public/PerformerComponent.h
--- a/Engine/Plugins/Animation/PerformanceCaptureCore/Source/PerformanceCaptureCore/Public/PerformerComponent.h
+++ b/Engine/Plugins/Animation/PerformanceCaptureCore/Source/PerformanceCaptureCore/Public/PerformerComponent.h
@@ -11,7 +11,7 @@
 class ULiveLinkInstance;
 
 UCLASS(BlueprintType, DisplayName="Performer Component",  ClassGroup=("Performance Capture"), meta=(BlueprintSpawnableComponent))
-class UPerformerComponent : public UActorComponent
+class PERFORMANCECAPTURECORE_API UPerformerComponent : public UActorComponent
 {
 	GENERATED_BODY()
 
diff --git a/Engine/Plugins/Animation/PerformanceCaptureCore/Source/PerformanceCaptureCore/Public/RetargetComponent.h b/Engine/Plugins/Animation/PerformanceCaptureCore/Source/PerformanceCaptureCore/Public/RetargetComponent.h
--- a/Engine/Plugins/Animation/PerformanceCaptureCore/Source/PerformanceCaptureCore/Public/RetargetComponent.h
+++ b/Engine/Plugins/Animation/PerformanceCaptureCore/Source/PerformanceCaptureCore/Public/RetargetComponent.h
@@ -12,7 +12,7 @@
 class UIKRetargeter;
 
 UCLASS(BlueprintType, ClassGroup=("Performance Capture"), meta=(BlueprintSpawnableComponent), DisplayName = "Retarget Component")
-class URetargetComponent : public UActorComponent
+class PERFORMANCECAPTURECORE_API URetargetComponent : public UActorComponent
 {
 	GENERATED_BODY()
 
