Performance Capture Core:
- Added Performance Capture|Sub categories to blueprint accessible functions

diff --git a/Engine/Plugins/Animation/PerformanceCaptureCore/Source/PerformanceCaptureCore/Public/CaptureCharacter.h b/Engine/Plugins/Animation/PerformanceCaptureCore/Source/PerformanceCaptureCore/Public/CaptureCharacter.h
--- a/Engine/Plugins/Animation/PerformanceCaptureCore/Source/PerformanceCaptureCore/Public/CaptureCharacter.h
+++ b/Engine/Plugins/Animation/PerformanceCaptureCore/Source/PerformanceCaptureCore/Public/CaptureCharacter.h
@@ -23,54 +23,54 @@ public:
 	/**
 	* CapturePerformer Actor that will be the source for retargeting.
 	*/
-	UPROPERTY(BlueprintReadOnly, EditAnywhere, Category = "Performance Capture")
+	UPROPERTY(BlueprintReadOnly, EditAnywhere, Category = "Performance Capture|Character")
 	TObjectPtr<ACapturePerformer> SourcePerformer;
 
 	/**
 	* The IKRetarget Asset to use for retargeting between the SourcePerformer and this Character.
 	*/
-	UPROPERTY(BlueprintReadOnly, EditAnywhere, Category = "Performance Capture")
+	UPROPERTY(BlueprintReadOnly, EditAnywhere, Category = "Performance Capture|Character")
 	TObjectPtr<UIKRetargeter> RetargetAsset;
 
 	/**
 	* Force all skeletal meshes to use the root SkeletalMesh as their Leader. Default = True.
 	*/
-	UPROPERTY(BlueprintReadOnly, EditAnywhere, Category = "Performance Capture")
+	UPROPERTY(BlueprintReadOnly, EditAnywhere, Category = "Performance Capture|Character")
 	bool bForceAllSkeletalMeshesToFollowLeader = true;
 	
 	/**
 	* Set a Custom Retarget Profile.
 	* @param InProfile New Retarget Profile.
 	*/
-	UFUNCTION(BlueprintCallable, Category = "Performance Capture")
+	UFUNCTION(BlueprintCallable, Category = "Performance Capture|Character")
 	void SetCustomRetargetProfile(FRetargetProfile InProfile);
 
 	/**
 	* Get a Custom Retarget Profile.
 	* @return Current Custom Retarget Profile.
 	*/
-	UFUNCTION(BlueprintCallable, Category = "Performance Capture")
+	UFUNCTION(BlueprintCallable, Category = "Performance Capture|Character")
 	FRetargetProfile GetCustomRetargetProfile();
 
 	/**
 	* Set the Source CapturePerformer Actor.
 	* @param InPerformer New Performer.
 	*/
-	UFUNCTION(BlueprintCallable, Category = "Performance Capture")
+	UFUNCTION(BlueprintCallable, Category = "Performance Capture|Character")
 	void SetSourcePerformer(ACapturePerformer* InPerformer);
 
 	/**
 	* Set the Retarget Asset.
 	* @param InRetargetAsset New IKRetarget Asset.
 	*/
-	UFUNCTION(BlueprintCallable, Category = "Performance Capture")
+	UFUNCTION(BlueprintCallable, Category = "Performance Capture|Character")
 	void SetRetargetAsset(UIKRetargeter* InRetargetAsset);
 
 	/**
 	* Force all Skeletal Meshes to follow the Controlled Skeletal mesh. 
 	* @param InFollowLeader New Bool.
 	*/
-	UFUNCTION(BlueprintCallable, Category = "Performance Capture")
+	UFUNCTION(BlueprintCallable, Category = "Performance Capture|Character")
 	void SetForceAllSkeletalMeshesToFollowLeader (bool InFollowLeader);
 
 protected:
diff --git a/Engine/Plugins/Animation/PerformanceCaptureCore/Source/PerformanceCaptureCore/Public/CapturePerformer.h b/Engine/Plugins/Animation/PerformanceCaptureCore/Source/PerformanceCaptureCore/Public/CapturePerformer.h
--- a/Engine/Plugins/Animation/PerformanceCaptureCore/Source/PerformanceCaptureCore/Public/CapturePerformer.h
+++ b/Engine/Plugins/Animation/PerformanceCaptureCore/Source/PerformanceCaptureCore/Public/CapturePerformer.h
@@ -42,35 +42,35 @@ public:
 	* Set the LiveLink Subject Name. Subject must have the Animation Role Type.
 	* @param Subject New LiveLink Subject.
 	*/
-	UFUNCTION(BlueprintCallable, Category = "Performance Capture")
+	UFUNCTION(BlueprintCallable, Category = "Performance Capture|Performer")
 	void SetLiveLinkSubject(FLiveLinkSubjectName Subject);
 
 	/**
 	* Get the LiveLink Subject Name.
 	* @return FLiveLinkSubjectName Current LiveLink Subject.
 	*/
-	UFUNCTION(BlueprintCallable, Category ="Performance Capture")
+	UFUNCTION(BlueprintCallable, Category ="Performance Capture|Performer")
 	FLiveLinkSubjectName GetLiveLinkSubject() const;
 
 	/**
 	* Set the LiveLink data to update the Skeletal Mesh pose.
 	* @param bEvaluateLinkLink Drive or pause the Skeletal Mesh from LiveLink Subject data.
 	*/
-	UFUNCTION(BlueprintCallable, Category ="Performance Capture")
+	UFUNCTION(BlueprintCallable, Category ="Performance Capture|Performer")
 	void SetEvaluateLiveLinkData(bool bEvaluateLinkLink);
 
 	/**
 	* Get state of LiveLink data evaluation. True = animation is updated.
 	* @return bool Is LiveLink data being evaluated.
 	*/
-	UFUNCTION(BlueprintPure, Category ="Performance Capture")
+	UFUNCTION(BlueprintPure, Category ="Performance Capture|Performer")
 	bool GetEvaluateLiveLinkData();
 
 	/**
 	* Set the Skeletal Mesh Asset of the root Skeletal Mesh.
 	* @param MocapMesh New Skeletal Mesh Asset.
 	*/
-	UFUNCTION(BlueprintCallable, Category = "Performance Capture")
+	UFUNCTION(BlueprintCallable, Category = "Performance Capture|Performer")
 	void SetMocapMesh(USkeletalMesh* MocapMesh);
 
 protected:
diff --git a/Engine/Plugins/Animation/PerformanceCaptureCore/Source/PerformanceCaptureCore/Public/PerformerComponent.h b/Engine/Plugins/Animation/PerformanceCaptureCore/Source/PerformanceCaptureCore/Public/PerformerComponent.h
--- a/Engine/Plugins/Animation/PerformanceCaptureCore/Source/PerformanceCaptureCore/Public/PerformerComponent.h
+++ b/Engine/Plugins/Animation/PerformanceCaptureCore/Source/PerformanceCaptureCore/Public/PerformerComponent.h
@@ -47,28 +47,28 @@ public:
 	* Set the LiveLink Subject Name. Subject must have the Animation Role Type.
 	* @param Subject New LiveLink Subject.
 	*/
-	UFUNCTION(BlueprintCallable, Category = "Performance Capture")
+	UFUNCTION(BlueprintCallable, Category = "Performance Capture|Performer")
 	void SetLiveLinkSubject(FLiveLinkSubjectName Subject);
 	
 	/**
 	* Get the LiveLink Subject Name.
 	* @return FLiveLinkSubjectName Current LiveLink Subject.
 	*/
-	UFUNCTION(BlueprintCallable, Category ="Performance Capture")
+	UFUNCTION(BlueprintCallable, Category ="Performance Capture|Performer")
 	FLiveLinkSubjectName GetLiveLinkSubject() const;
 	
 	/**
 	* Set the LiveLink data to update the Skeletal Mesh pose.
 	* @param bEvaluateLinkLink Drive or pause the Skeletal Mesh from LiveLink Subject data.
 	*/
-	UFUNCTION(BlueprintCallable, Category ="Performance Capture")
+	UFUNCTION(BlueprintCallable, Category ="Performance Capture|Performer")
 	void SetEvaluateLiveLinkData(bool bEvaluateLinkLink);
 	
 	/**
 	* Get the LiveLink Subject Name. Subject must have the Animation Role Type.
 	* @return bool Is LiveLink data being evaluated.
 	*/
-	UFUNCTION(BlueprintPure, Category ="Performance Capture")
+	UFUNCTION(BlueprintPure, Category ="Performance Capture|Performer")
 	bool GetEvaluateLiveLinkData();
 
 protected:
diff --git a/Engine/Plugins/Animation/PerformanceCaptureCore/Source/PerformanceCaptureCore/Public/RetargetComponent.h b/Engine/Plugins/Animation/PerformanceCaptureCore/Source/PerformanceCaptureCore/Public/RetargetComponent.h
--- a/Engine/Plugins/Animation/PerformanceCaptureCore/Source/PerformanceCaptureCore/Public/RetargetComponent.h
+++ b/Engine/Plugins/Animation/PerformanceCaptureCore/Source/PerformanceCaptureCore/Public/RetargetComponent.h
@@ -54,35 +54,35 @@ public:
 	* Set the Source Performer Mesh.
 	* @param InPerformerMesh New Source Skeletal Mesh.
 	*/
-	UFUNCTION(BlueprintCallable, Category = "Performance Capture")
+	UFUNCTION(BlueprintCallable, Category = "Performance Capture|Retarget")
 	void SetSourcePerformerMesh(USkeletalMeshComponent* InPerformerMesh);
 	
 	/**
 	* Set the Controlled Skeletal Mesh.
 	* @param InControlledMesh New Controlled Skeletal Mesh Component.
 	*/
-	UFUNCTION(BlueprintCallable, Category = "Performance Capture")
+	UFUNCTION(BlueprintCallable, Category = "Performance Capture|Retarget")
 	void SetControlledMesh(USkeletalMeshComponent* InControlledMesh);
 
 	/**
 	* Set the Retarget Asset.
 	* @param InRetargetAsset New IKRetarget Asset.
 	*/
-	UFUNCTION(BlueprintCallable, Category = "Performance Capture")
+	UFUNCTION(BlueprintCallable, Category = "Performance Capture|Retarget")
 	void SetRetargetAsset(UIKRetargeter* InRetargetAsset);
 
 	/**
 	* Set a Custom Retarget Profile.
 	* @param InProfile New Retarget Profile.
 	*/
-	UFUNCTION(BlueprintCallable, Category = "Performance Capture")
+	UFUNCTION(BlueprintCallable, Category = "Performance Capture|Retarget")
 	void SetCustomRetargetProfile(FRetargetProfile InProfile);
 
 	/**
 	* Get Retarget Profile.
 	* @return Current Custom Retarget Profile.
 	*/
-	UFUNCTION(BlueprintCallable, Category = "Performance Capture")
+	UFUNCTION(BlueprintCallable, Category = "Performance Capture|Retarget")
 	FRetargetProfile GetCustomRetargetProfile();
 	
 protected:
