Performance Capture Core

-Fix for blueprint exposed function not correctly setting child skelmeshes to follow leader. Fix requires call to function on the retarget component that iterates through all child skelmeshes

[RN] Fixed bug in Capture Character SetForceAllSkeletalMeshesTofFollowLeader function

diff --git a/Engine/Plugins/Animation/PerformanceCaptureCore/Source/PerformanceCaptureCore/Private/CaptureCharacter.cpp b/Engine/Plugins/Animation/PerformanceCaptureCore/Source/PerformanceCaptureCore/Private/CaptureCharacter.cpp
--- a/Engine/Plugins/Animation/PerformanceCaptureCore/Source/PerformanceCaptureCore/Private/CaptureCharacter.cpp
+++ b/Engine/Plugins/Animation/PerformanceCaptureCore/Source/PerformanceCaptureCore/Private/CaptureCharacter.cpp
@@ -50,12 +50,14 @@ void ACaptureCharacter::SetRetargetAsset(UIKRetargeter* InRetargetAsset)
 {
 	RetargetAsset = InRetargetAsset;
 	RetargetComponent->SetRetargetAsset(InRetargetAsset);
+	RetargetComponent->InitiateAnimation();
 }
 
 void ACaptureCharacter::SetForceAllSkeletalMeshesToFollowLeader(bool InFollowLeader)
 {
 	bForceAllSkeletalMeshesToFollowLeader = InFollowLeader;
-	RetargetComponent->bForceOtherMeshesToFollowControlledMesh = bForceAllSkeletalMeshesToFollowLeader;
+	RetargetComponent->SetForceOtherMeshesToFollowControlledMesh(InFollowLeader);
+	RetargetComponent->InitiateAnimation();
 }
 
 void ACaptureCharacter::PostRegisterAllComponents()
