Performance Capture Core

- Capture Character: Swapped performer reference to Soft ref as hard ref was causing an bad reference when saving a level sequence if the Capture Character is spawnable in a level sequence. Swap to soft ref removes hard ref save failure.

#rnx

diff --git a/Engine/Plugins/Animation/PerformanceCaptureCore/Source/PerformanceCaptureCore/Private/CaptureCharacter.cpp b/Engine/Plugins/Animation/PerformanceCaptureCore/Source/PerformanceCaptureCore/Private/CaptureCharacter.cpp
--- a/Engine/Plugins/Animation/PerformanceCaptureCore/Source/PerformanceCaptureCore/Private/CaptureCharacter.cpp
+++ b/Engine/Plugins/Animation/PerformanceCaptureCore/Source/PerformanceCaptureCore/Private/CaptureCharacter.cpp
@@ -66,7 +66,7 @@ void ACaptureCharacter::PostRegisterAllComponents()
 	if (!HasAnyFlags(RF_ClassDefaultObject))
 	{
 		RetargetComponent->SetControlledMesh(GetSkeletalMeshComponent());
-		SetSourcePerformer(SourcePerformer);
+		SetSourcePerformer(SourcePerformer.Get());
 		RetargetComponent->SetRetargetAsset(RetargetAsset);
 		RetargetComponent->SetForceOtherMeshesToFollowControlledMesh(bForceAllSkeletalMeshesToFollowLeader);
 		RetargetComponent->InitiateAnimation();
@@ -82,7 +82,7 @@ void ACaptureCharacter::PostEditChangeProperty(FPropertyChangedEvent& PropertyCh
 	{
 		if(PropertyChangedEvent.Property->GetFName() == GET_MEMBER_NAME_CHECKED(ACaptureCharacter, SourcePerformer))
 		{
-			SetSourcePerformer(SourcePerformer);
+			SetSourcePerformer(SourcePerformer.Get());
 			RetargetComponent->InitiateAnimation();
 		}
 		
diff --git a/Engine/Plugins/Animation/PerformanceCaptureCore/Source/PerformanceCaptureCore/Public/CaptureCharacter.h b/Engine/Plugins/Animation/PerformanceCaptureCore/Source/PerformanceCaptureCore/Public/CaptureCharacter.h
--- a/Engine/Plugins/Animation/PerformanceCaptureCore/Source/PerformanceCaptureCore/Public/CaptureCharacter.h
+++ b/Engine/Plugins/Animation/PerformanceCaptureCore/Source/PerformanceCaptureCore/Public/CaptureCharacter.h
@@ -24,7 +24,7 @@ public:
 	* CapturePerformer Actor that will be the source for retargeting.
 	*/
 	UPROPERTY(BlueprintReadOnly, EditAnywhere, Category = "Performance Capture|Character")
-	TObjectPtr<ACapturePerformer> SourcePerformer;
+	TSoftObjectPtr<ACapturePerformer> SourcePerformer;
 
 	/**
 	* The IKRetarget Asset to use for retargeting between the SourcePerformer and this Character.
diff --git a/Engine/Plugins/Animation/PerformanceCaptureCore/Source/PerformanceCaptureCore/Public/PerformerComponent.h b/Engine/Plugins/Animation/PerformanceCaptureCore/Source/PerformanceCaptureCore/Public/PerformerComponent.h
--- a/Engine/Plugins/Animation/PerformanceCaptureCore/Source/PerformanceCaptureCore/Public/PerformerComponent.h
+++ b/Engine/Plugins/Animation/PerformanceCaptureCore/Source/PerformanceCaptureCore/Public/PerformerComponent.h
@@ -65,7 +65,7 @@ public:
 	void SetEvaluateLiveLinkData(bool bEvaluateLinkLink);
 	
 	/**
-	* Get the LiveLink Subject Name. Subject must have the Animation Role Type.
+	* Get the LiveLink Evaluation State. Subject must have the Animation Role Type.
 	* @return bool Is LiveLink data being evaluated.
 	*/
 	UFUNCTION(BlueprintPure, Category ="Performance Capture|Performer")
