Performance Capture Core

- force setting of tick prerequisite actor on  the retarget mesh to always be the source performer so there is no longer a one frame delay between the motion of the performer source and the retargetted mesh.

diff --git a/Engine/Plugins/Animation/PerformanceCaptureCore/Source/PerformanceCaptureCore/Private/RetargetComponent.cpp b/Engine/Plugins/Animation/PerformanceCaptureCore/Source/PerformanceCaptureCore/Private/RetargetComponent.cpp
--- a/Engine/Plugins/Animation/PerformanceCaptureCore/Source/PerformanceCaptureCore/Private/RetargetComponent.cpp
+++ b/Engine/Plugins/Animation/PerformanceCaptureCore/Source/PerformanceCaptureCore/Private/RetargetComponent.cpp
@@ -238,6 +238,7 @@ void URetargetComponent::InitiateAnimation()
 			RetargetAsset->IncrementVersion();
 			AnimInstance->ConfigureAnimInstance(RetargetAsset, SourceMesh, CustomRetargetProfile);
 			ControlledMesh->SetUpdateAnimationInEditor(true);
+			ControlledMesh->AddTickPrerequisiteActor(SourceMesh->GetOwner());
 			ControlledMesh->bPropagateCurvesToFollowers = true;
 			ControlledMesh->VisibilityBasedAnimTickOption = EVisibilityBasedAnimTickOption::AlwaysTickPoseAndRefreshBones;
 			ControlledMesh->InitAnim(true /*bForceReinit*/);
