[Retarget Component] Fixed issue crash when nulling the source mesh component.

diff --git a/Engine/Plugins/Animation/PerformanceCaptureCore/Source/PerformanceCaptureCore/Private/RetargetComponent.cpp b/Engine/Plugins/Animation/PerformanceCaptureCore/Source/PerformanceCaptureCore/Private/RetargetComponent.cpp
--- a/Engine/Plugins/Animation/PerformanceCaptureCore/Source/PerformanceCaptureCore/Private/RetargetComponent.cpp
+++ b/Engine/Plugins/Animation/PerformanceCaptureCore/Source/PerformanceCaptureCore/Private/RetargetComponent.cpp
@@ -223,9 +223,9 @@ void URetargetComponent::SetForceOtherMeshesToFollowControlledMesh(bool bInBool)
 void URetargetComponent::InitiateAnimation()
 {
 	USkeletalMeshComponent* ControlledMesh = Cast<USkeletalMeshComponent> (ControlledSkeletalMeshComponent.GetComponent(GetOwner()));
-	USkeletalMeshComponent* SourceMesh = Cast<USkeletalMeshComponent>(SourceSkeletalMeshComponent.GetComponent(GetOwner()));
+	USkeletalMeshComponent* SourceMesh = Cast<USkeletalMeshComponent>(SourceSkeletalMeshComponent.GetComponent(nullptr));
 	
-	const bool bShouldBeReinitialized = IsValid(SourceMesh) && IsValid(RetargetAsset) && IsValid(ControlledMesh);
+	const bool bShouldBeReinitialized = IsValid(RetargetAsset) && IsValid(ControlledMesh);
 	
 	if(bShouldBeReinitialized)
 	{
@@ -233,13 +233,17 @@ void URetargetComponent::InitiateAnimation()
 		ControlledMesh->SetAnimInstanceClass(URetargetAnimInstance::StaticClass());
 
 		TObjectPtr<URetargetAnimInstance> AnimInstance = Cast<URetargetAnimInstance>(ControlledMesh->GetAnimInstance());
+
+		if (SourceMesh && SourceMesh->GetOwner())
+		{
+			ControlledMesh->AddTickPrerequisiteActor(SourceMesh->GetOwner());
+		}
 		
 		if(AnimInstance)
 		{
 			RetargetAsset->IncrementVersion();
 			AnimInstance->ConfigureAnimInstance(RetargetAsset, SourceMesh, CustomRetargetProfile);
 			ControlledMesh->SetUpdateAnimationInEditor(true);
-			ControlledMesh->AddTickPrerequisiteActor(SourceMesh->GetOwner());
 			ControlledMesh->bPropagateCurvesToFollowers = true;
 			ControlledMesh->VisibilityBasedAnimTickOption = EVisibilityBasedAnimTickOption::AlwaysTickPoseAndRefreshBones;
 			ControlledMesh->InitAnim(true /*bForceReinit*/);
