[Performance Capture Core]
- Remove BP write access on members that won't be reintiated by post-edit changes when written to from BP.
- Added new BP callable functions to update members that were previously BlueprintReadWrite and ensure reinitialization is called properly when updated.

[RN] Blueprint ReadWrite members in the performer and retarget components now have dedicated UFunctions to set variables that require animation reinitialisation

diff --git a/Engine/Plugins/Animation/PerformanceCaptureCore/Source/PerformanceCaptureCore/Private/PerformerComponent.cpp b/Engine/Plugins/Animation/PerformanceCaptureCore/Source/PerformanceCaptureCore/Private/PerformerComponent.cpp
--- a/Engine/Plugins/Animation/PerformanceCaptureCore/Source/PerformanceCaptureCore/Private/PerformerComponent.cpp
+++ b/Engine/Plugins/Animation/PerformanceCaptureCore/Source/PerformanceCaptureCore/Private/PerformerComponent.cpp
@@ -123,11 +123,28 @@ void UPerformerComponent::SetEvaluateLiveLinkData(bool bEvaluateLinkLink)
 	}
 }
 
-bool UPerformerComponent::GetEvaluateLiveLinkData()
+bool UPerformerComponent::GetEvaluateLiveLinkData() const
 {
 	return bEvaluateAnimation;
 }
 
+bool UPerformerComponent::SetControlledMesh(USkeletalMeshComponent* SkeletalMesh)
+{
+	if (SkeletalMesh->GetOwner() == this->GetOwner() && SkeletalMesh) //Check the given skeletal mesh component belongs to the same actor as this performer component and is valid
+	{
+		ControlledSkeletalMesh.OverrideComponent = SkeletalMesh;
+		bIsDirty=true; // Mark bIsDirty true so that next frame InitiateAnimation is called.
+		return true;
+	}
+	return false;
+}
+
+void UPerformerComponent::SetForceMeshesFollowLeader(bool bInBool)
+{
+	bForceOtherMeshesToFollowControlledMesh = bInBool;
+	bIsDirty=true; // Mark bIsDirty true so that next frame InitiateAnimation is called.
+}
+
 void UPerformerComponent::InitiateAnimation()
 {
 	TArray<USkeletalMeshComponent*> OwnerSkeletalMeshComponents;
@@ -185,4 +202,4 @@ void UPerformerComponent::InitiateAnimation()
 		SetEvaluateLiveLinkData(bEvaluateAnimation);
 	}	
 	bIsDirty=false;
-}
\ No newline at end of file
+}
diff --git a/Engine/Plugins/Animation/PerformanceCaptureCore/Source/PerformanceCaptureCore/Private/RetargetComponent.cpp b/Engine/Plugins/Animation/PerformanceCaptureCore/Source/PerformanceCaptureCore/Private/RetargetComponent.cpp
--- a/Engine/Plugins/Animation/PerformanceCaptureCore/Source/PerformanceCaptureCore/Private/RetargetComponent.cpp
+++ b/Engine/Plugins/Animation/PerformanceCaptureCore/Source/PerformanceCaptureCore/Private/RetargetComponent.cpp
@@ -47,6 +47,12 @@ FRetargetProfile URetargetComponent::GetCustomRetargetProfile()
 	return OutProfile;
 }
 
+void URetargetComponent::SetForceMeshesFollowLeader(const bool bInBool)
+{
+	bForceOtherMeshesToFollowControlledMesh = bInBool;
+	SetForceOtherMeshesToFollowControlledMesh(bForceOtherMeshesToFollowControlledMesh);
+}
+
 void URetargetComponent::SetSourcePerformerMesh(USkeletalMeshComponent* InPerformerMesh)
 {
 	SourceSkeletalMeshComponent.OverrideComponent = InPerformerMesh;
diff --git a/Engine/Plugins/Animation/PerformanceCaptureCore/Source/PerformanceCaptureCore/Public/PerformerComponent.h b/Engine/Plugins/Animation/PerformanceCaptureCore/Source/PerformanceCaptureCore/Public/PerformerComponent.h
--- a/Engine/Plugins/Animation/PerformanceCaptureCore/Source/PerformanceCaptureCore/Public/PerformerComponent.h
+++ b/Engine/Plugins/Animation/PerformanceCaptureCore/Source/PerformanceCaptureCore/Public/PerformerComponent.h
@@ -22,25 +22,25 @@ public:
 	/**
 	* LiveLink Subject Name. Must have the Animation Role Type.
 	*/
-	UPROPERTY(BlueprintReadWrite, EditAnywhere, Category = "Performance Capture")
+	UPROPERTY(BlueprintReadOnly, EditAnywhere, Category = "Performance Capture")
 	FLiveLinkSubjectName SubjectName;
 
 	/**
 	* The Skeletal Mesh driven by the LiveLink subject. Skeleton mush be compatible with the LiveLink subject's bone hierarchy. 
 	*/
-	UPROPERTY(BlueprintReadWrite, EditInstanceOnly, Category = "Performance Capture", meta=(UseComponentPicker, AllowedClasses = "/Script/Engine.SkeletalMeshComponent"))
+	UPROPERTY(BlueprintReadOnly, EditInstanceOnly, Category = "Performance Capture", meta=(UseComponentPicker, AllowedClasses = "/Script/Engine.SkeletalMeshComponent"))
 	FComponentReference ControlledSkeletalMesh;
 	
 	/**
 	* Evaluate LiveLink animation. Set to false to pause animation.
 	*/
-	UPROPERTY(BlueprintReadWrite, EditAnywhere, Category = "Performance Capture")
+	UPROPERTY(EditAnywhere, Category = "Performance Capture")
 	bool bEvaluateAnimation = true;
 	
 	/**
 	* Force all other skeletal meshes in the Owner Actor to follow the pose of the Controlled Skeletal Mesh.
 	*/
-	UPROPERTY(BlueprintReadWrite, EditAnywhere, Category = "Performance Capture")
+	UPROPERTY(BlueprintReadOnly, EditAnywhere, Category = "Performance Capture")
 	bool bForceOtherMeshesToFollowControlledMesh = true;
 
 	/**
@@ -69,7 +69,22 @@ public:
 	* @return bool Is LiveLink data being evaluated.
 	*/
 	UFUNCTION(BlueprintPure, Category ="Performance Capture|Performer")
-	bool GetEvaluateLiveLinkData();
+	bool GetEvaluateLiveLinkData() const;
+
+	/**
+	* Set the controlled skeletal mesh component. Must be a member of the owner actor of this component.
+	* @param SkeletalMesh New SkeletalMesh.
+	* @return bool Was the given skeletal mesh component successfully applied.
+	*/
+	UFUNCTION(BlueprintCallable, Category = "Performance Capture|Performer")
+	bool SetControlledMesh(USkeletalMeshComponent* SkeletalMesh);
+
+	/**
+	* Set the bForceOtherMeshesToFollowControlledMesh variable. Changing will re-initialise animation.
+	* @param bInBool DForce all other skeletal meshes in the Owner Actor to follow the pose of the Controlled Skeletal Mesh.
+	*/
+	UFUNCTION(BlueprintCallable, Category ="Performance Capture|Performer")
+	void SetForceMeshesFollowLeader(bool bInBool);
 
 protected:
 	virtual void DestroyComponent(bool bPromoteChildren) override;
diff --git a/Engine/Plugins/Animation/PerformanceCaptureCore/Source/PerformanceCaptureCore/Public/RetargetComponent.h b/Engine/Plugins/Animation/PerformanceCaptureCore/Source/PerformanceCaptureCore/Public/RetargetComponent.h
--- a/Engine/Plugins/Animation/PerformanceCaptureCore/Source/PerformanceCaptureCore/Public/RetargetComponent.h
+++ b/Engine/Plugins/Animation/PerformanceCaptureCore/Source/PerformanceCaptureCore/Public/RetargetComponent.h
@@ -23,19 +23,19 @@ public:
 	/**
 	* Skeletal Mesh component that will be the source for retargeting. Can be on the Owner Actor or another Actor in the same level.
 	*/
-	UPROPERTY(BlueprintReadWrite, EditInstanceOnly, Category = "Performance Capture",  meta = (UseComponentPicker, AllowAnyActor, AllowedClasses = "/Script/Engine.SkeletalMeshComponent"))
+	UPROPERTY(BlueprintReadOnly, EditInstanceOnly, Category = "Performance Capture",  meta = (UseComponentPicker, AllowAnyActor, AllowedClasses = "/Script/Engine.SkeletalMeshComponent"))
 	FComponentReference SourceSkeletalMeshComponent ;
 
 	/**
 	* Skeletal Mesh that will be driven by the IKRetargeter. Limited to skeletal meshes on this component's Owner Actor.
 	*/
-	UPROPERTY(BlueprintReadWrite, EditAnywhere,Category = "Performance Capture", meta = (UseComponentPicker, AllowedClasses = "/Script/Engine.SkeletalMeshComponent"))
+	UPROPERTY(BlueprintReadOnly, EditAnywhere,Category = "Performance Capture", meta = (UseComponentPicker, AllowedClasses = "/Script/Engine.SkeletalMeshComponent"))
 	FComponentReference ControlledSkeletalMeshComponent;
 
 	/**
 	* Force all skeletal meshes to use the ControlledSkeletalMeshComponent as their Leader. Default = True.
 	*/
-	UPROPERTY(BlueprintReadWrite, EditAnywhere, Category = "Performance Capture")
+	UPROPERTY(BlueprintReadOnly, EditAnywhere, Category = "Performance Capture")
 	bool bForceOtherMeshesToFollowControlledMesh = true;
 
 	/**
@@ -47,7 +47,7 @@ public:
 	/**
 	* Custom Retarget Profile. Should be used to override Retarget settings on the RetargetAsset.
 	*/
-	UPROPERTY(BlueprintReadWrite, Interp, EditAnywhere, Category = "Performance Capture")
+	UPROPERTY(BlueprintReadOnly, Interp, EditAnywhere, Category = "Performance Capture")
 	FRetargetProfile CustomRetargetProfile;
 	
 	/**
@@ -84,6 +84,13 @@ public:
 	*/
 	UFUNCTION(BlueprintCallable, Category = "Performance Capture|Retarget")
 	FRetargetProfile GetCustomRetargetProfile();
+
+	/**
+	* Set the bForceOtherMeshesToFollowControlledMesh variable. Changing will re-initialise animation.
+	* @param bInBool Force all other skeletal meshes in the Owner Actor to follow the pose of the Controlled Skeletal Mesh.
+	*/
+	UFUNCTION(BlueprintCallable, Category = "Performance Capture|Retarget")
+	void SetForceMeshesFollowLeader(bool bInBool);
 	
 protected:
 	// Called when the game starts
