[Performance Capture Core]
- Fix SA warning.

diff --git a/Engine/Plugins/Animation/PerformanceCaptureCore/Source/PerformanceCaptureCore/Private/PerformerComponent.cpp b/Engine/Plugins/Animation/PerformanceCaptureCore/Source/PerformanceCaptureCore/Private/PerformerComponent.cpp
--- a/Engine/Plugins/Animation/PerformanceCaptureCore/Source/PerformanceCaptureCore/Private/PerformerComponent.cpp
+++ b/Engine/Plugins/Animation/PerformanceCaptureCore/Source/PerformanceCaptureCore/Private/PerformerComponent.cpp
@@ -130,12 +130,16 @@ bool UPerformerComponent::GetEvaluateLiveLinkData() const
 
 bool UPerformerComponent::SetControlledMesh(USkeletalMeshComponent* SkeletalMesh)
 {
-	if (SkeletalMesh->GetOwner() == this->GetOwner() && SkeletalMesh) //Check the given skeletal mesh component belongs to the same actor as this performer component and is valid
+	if (SkeletalMesh)
 	{
-		ControlledSkeletalMesh.OverrideComponent = SkeletalMesh;
-		bIsDirty=true; // Mark bIsDirty true so that next frame InitiateAnimation is called.
-		return true;
+		if (SkeletalMesh->GetOwner() == this->GetOwner()) //Check the given skeletal mesh component belongs to the same actor as this performer component and is valid
+		{
+			ControlledSkeletalMesh.OverrideComponent = SkeletalMesh;
+			bIsDirty=true; // Mark bIsDirty true so that next frame InitiateAnimation is called.
+			return true;
+		}
 	}
+
 	return false;
 }
 
