[Performance Capture Core]
- Fix potential crash in both the Retarget and Performer components when used in the actors that do not just have skeletal mesh components in their hierarchy while using the BP editor eg: metahumans.
- Fix ambiguous GetComponents() call to specifically only get skeletal mesh components.
GetOwner() can sometimes be invalid while a BP is recompiling and this change checks for that.

[RN] Fix potential crash in Retarget Component when used with Metahuman BP

diff --git a/Engine/Plugins/Animation/PerformanceCaptureCore/Source/PerformanceCaptureCore/Private/PerformerComponent.cpp b/Engine/Plugins/Animation/PerformanceCaptureCore/Source/PerformanceCaptureCore/Private/PerformerComponent.cpp
--- a/Engine/Plugins/Animation/PerformanceCaptureCore/Source/PerformanceCaptureCore/Private/PerformerComponent.cpp
+++ b/Engine/Plugins/Animation/PerformanceCaptureCore/Source/PerformanceCaptureCore/Private/PerformerComponent.cpp
@@ -18,7 +18,7 @@ void UPerformerComponent::DestroyComponent(bool bPromoteChildren)
 {
 	//Reset all the skeletal meshes in the owner actor 
 	TArray<USkeletalMeshComponent*> SkeletalMeshComponents;
-	GetOwner()->GetComponents(SkeletalMeshComponents);
+	GetOwner()->GetComponents<USkeletalMeshComponent>(SkeletalMeshComponents);
 	for (USkeletalMeshComponent* SkeletalMeshComponent : SkeletalMeshComponents)
 	{
 		SkeletalMeshComponent->SetUpdateAnimationInEditor(false/* NewUpdateState */);
@@ -152,7 +152,7 @@ void UPerformerComponent::SetForceMeshesFollowLeader(bool bInBool)
 void UPerformerComponent::InitiateAnimation()
 {
 	TArray<USkeletalMeshComponent*> OwnerSkeletalMeshComponents;
-	GetOwner()->GetComponents(OwnerSkeletalMeshComponents);
+	GetOwner()->GetComponents<USkeletalMeshComponent>(OwnerSkeletalMeshComponents);
 
 	//Set the mesh that will be controlled
 	USkeletalMeshComponent* ControlledMesh = Cast<USkeletalMeshComponent>(ControlledSkeletalMesh.GetComponent(GetOwner()));
diff --git a/Engine/Plugins/Animation/PerformanceCaptureCore/Source/PerformanceCaptureCore/Private/RetargetComponent.cpp b/Engine/Plugins/Animation/PerformanceCaptureCore/Source/PerformanceCaptureCore/Private/RetargetComponent.cpp
--- a/Engine/Plugins/Animation/PerformanceCaptureCore/Source/PerformanceCaptureCore/Private/RetargetComponent.cpp
+++ b/Engine/Plugins/Animation/PerformanceCaptureCore/Source/PerformanceCaptureCore/Private/RetargetComponent.cpp
@@ -123,7 +123,7 @@ void URetargetComponent::DestroyComponent(bool bPromoteChildren)
 {
 	//When the component is removed, reset the state of the skelmeshes in the owner actor
 	TArray<USkeletalMeshComponent*> OwnerSkeletalMeshComponents;
-	GetOwner()->GetComponents(OwnerSkeletalMeshComponents);
+	GetOwner()->GetComponents<USkeletalMeshComponent>(OwnerSkeletalMeshComponents);
 	
 	for (USkeletalMeshComponent* OwnerSkeletalMeshComponent : OwnerSkeletalMeshComponents)
 	{
@@ -187,8 +187,12 @@ void URetargetComponent::SetForceOtherMeshesToFollowControlledMesh(bool bInBool)
 	const USkeletalMeshComponent* SourceMesh = Cast<USkeletalMeshComponent>(SourceSkeletalMeshComponent.GetComponent(GetOwner()));
 	
 	TArray<USkeletalMeshComponent*> OwnerSkeletalMeshComponents;
-	GetOwner()->GetComponents(OwnerSkeletalMeshComponents);
-
+	//Check the owner is valid because it can be Null if being reconstructed in the BP editor.
+	if(GetOwner()) 
+	{
+		GetOwner()->GetComponents<USkeletalMeshComponent>(OwnerSkeletalMeshComponents);
+	}
+	
 	if(IsValid(ControlledMesh) && IsValid(SourceMesh))
 	{
 		for (USkeletalMeshComponent* OwnerSkeletalMeshComponent : OwnerSkeletalMeshComponents)
