MocapManager: Updates for MocapManager to be able to use any Skeletal Mesh Actor as a Character in addition to ACaptureCharacter.

- ACaptureCharacter's SourcePerformer, RetargetAsset and bForceAllSkeletalMeshesToFollowLeader are deprecated in favor of the same properties on Retarget Component.
- ACaptureCharacter function calls now all route through to its Retarget Component.
- Type of CaptureCharacterClass member on PCapCharacterDataAsset has been changed from TSoftClassPtr<ACaptureCharacter> to  TSoftClassPtr<ASkeletalMeshActor>.
- AdditionalMeshes property on PCapCharacterDataAsset has been deprecated for removal.
- URetargetComponent has been updated with a SourcePerformer property and to construct itself properly for retargeting on any Skeletal Mesh Actor.
- PCapActorFactory updated to use the Retarget Component and Skeletal Mesh Actor instead of ACaptureCharacter. A Retarget Component is added post spawn if one is not already present on the spawned Skeletal Mesh Actor.

diff --git a/Engine/Plugins/Animation/PerformanceCaptureCore/Source/PerformanceCaptureCore/Private/CaptureCharacter.cpp b/Engine/Plugins/Animation/PerformanceCaptureCore/Source/PerformanceCaptureCore/Private/CaptureCharacter.cpp
--- a/Engine/Plugins/Animation/PerformanceCaptureCore/Source/PerformanceCaptureCore/Private/CaptureCharacter.cpp
+++ b/Engine/Plugins/Animation/PerformanceCaptureCore/Source/PerformanceCaptureCore/Private/CaptureCharacter.cpp
@@ -14,8 +14,6 @@ ACaptureCharacter::ACaptureCharacter()
 	// Set this actor to call Tick() every frame.  You can turn this off to improve performance if you don't need it.
 	PrimaryActorTick.bCanEverTick = true;
 	RetargetComponent = CreateDefaultSubobject<URetargetComponent>("RetargetComponent");
-	RetargetComponent->ControlledSkeletalMeshComponent.OverrideComponent = Cast<USkeletalMeshComponent>(GetRootComponent());
-	RetargetComponent->SourceSkeletalMeshComponent.OverrideComponent = nullptr;
 }
 
 
@@ -28,78 +26,57 @@ void ACaptureCharacter::SetCustomRetargetProfile(FRetargetProfile InProfile)
 //Get the custom retarget profile
 FRetargetProfile ACaptureCharacter::GetCustomRetargetProfile()
 {
-	FRetargetProfile OutProfile = RetargetComponent->GetCustomRetargetProfile();
-	return OutProfile;
+	return RetargetComponent->GetCustomRetargetProfile();
 }
 
 void ACaptureCharacter::SetSourcePerformer(ACapturePerformer* InPerformer)
 {
-	if(IsValid(InPerformer))
-	{
-		SourcePerformer = InPerformer;
-		RetargetComponent->SetSourcePerformerMesh(InPerformer->GetSkeletalMeshComponent());
-		RetargetComponent->InitiateAnimation();
-	}
-	else
-	{
-		SourcePerformer = nullptr;
-		RetargetComponent->SetSourcePerformerMesh(GetSkeletalMeshComponent());
-		RetargetComponent->InitiateAnimation();
-	}
+	RetargetComponent->SetSourcePerformer(InPerformer);
 }
 
 //Update the retarget asset
 void ACaptureCharacter::SetRetargetAsset(UIKRetargeter* InRetargetAsset)
 {
-	RetargetAsset = InRetargetAsset;
 	RetargetComponent->SetRetargetAsset(InRetargetAsset);
-	RetargetComponent->InitiateAnimation();
 }
 
 void ACaptureCharacter::SetForceAllSkeletalMeshesToFollowLeader(bool InFollowLeader)
 {
-	bForceAllSkeletalMeshesToFollowLeader = InFollowLeader;
-	RetargetComponent->SetForceOtherMeshesToFollowControlledMesh(InFollowLeader);
-	RetargetComponent->InitiateAnimation();
+	RetargetComponent->SetForceMeshesFollowLeader(InFollowLeader);
 }
 
-void ACaptureCharacter::PostRegisterAllComponents()
+URetargetComponent* ACaptureCharacter::GetRetargetComponent() const
 {
-	Super::PostRegisterAllComponents();
-	if (!HasAnyFlags(RF_ClassDefaultObject))
-	{
-		RetargetComponent->SetControlledMesh(GetSkeletalMeshComponent());
-		SetSourcePerformer(SourcePerformer.Get());
-		RetargetComponent->SetRetargetAsset(RetargetAsset);
-		RetargetComponent->SetForceOtherMeshesToFollowControlledMesh(bForceAllSkeletalMeshesToFollowLeader);
-		RetargetComponent->InitiateAnimation();
-	}
+	return RetargetComponent;
 }
 
-#if WITH_EDITOR
-void ACaptureCharacter::PostEditChangeProperty(FPropertyChangedEvent& PropertyChangedEvent)
+void ACaptureCharacter::PostLoad()
 {
-	Super::PostEditChangeProperty(PropertyChangedEvent);
+	Super::PostLoad();
+
+#if WITH_EDITORONLY_DATA
+	// Fix up deprecation values
+PRAGMA_DISABLE_DEPRECATION_WARNINGS
 
-	if(PropertyChangedEvent.Property != nullptr)
+	if (SourcePerformer_DEPRECATED.IsPending())
 	{
-		if(PropertyChangedEvent.Property->GetFName() == GET_MEMBER_NAME_CHECKED(ACaptureCharacter, SourcePerformer))
-		{
-			SetSourcePerformer(SourcePerformer.Get());
-			RetargetComponent->InitiateAnimation();
-		}
-		
-		if(PropertyChangedEvent.Property->GetFName() == GET_MEMBER_NAME_CHECKED(ACaptureCharacter, RetargetAsset))
-		{
-			SetRetargetAsset(RetargetAsset);
-			RetargetComponent->InitiateAnimation();
-		}
-		
-		if(PropertyChangedEvent.Property->GetFName() == GET_MEMBER_NAME_CHECKED(ACaptureCharacter, bForceAllSkeletalMeshesToFollowLeader))
-		{
-			RetargetComponent->bForceOtherMeshesToFollowControlledMesh = bForceAllSkeletalMeshesToFollowLeader;
-			RetargetComponent->InitiateAnimation();
-		}
+		RetargetComponent->SourcePerformer = SourcePerformer_DEPRECATED;
+		SourcePerformer_DEPRECATED.Reset();
 	}
+
+	if (RetargetAsset_DEPRECATED != nullptr)
+	{
+		RetargetComponent->RetargetAsset = RetargetAsset_DEPRECATED;
+		RetargetAsset_DEPRECATED = nullptr;
+	}
+
+	// bForceAllSkeletalMeshesToFollowLeader defaults to true. If it is therefore not true, set the false value across
+	if (!bForceAllSkeletalMeshesToFollowLeader_DEPRECATED)
+	{
+		RetargetComponent->bForceOtherMeshesToFollowControlledMesh = false;
+		bForceAllSkeletalMeshesToFollowLeader_DEPRECATED = true;
+	}
+
+PRAGMA_ENABLE_DEPRECATION_WARNINGS
+#endif // WITH_EDITORONLY_DATA
 }
-#endif
diff --git a/Engine/Plugins/Animation/PerformanceCaptureCore/Source/PerformanceCaptureCore/Private/RetargetComponent.cpp b/Engine/Plugins/Animation/PerformanceCaptureCore/Source/PerformanceCaptureCore/Private/RetargetComponent.cpp
--- a/Engine/Plugins/Animation/PerformanceCaptureCore/Source/PerformanceCaptureCore/Private/RetargetComponent.cpp
+++ b/Engine/Plugins/Animation/PerformanceCaptureCore/Source/PerformanceCaptureCore/Private/RetargetComponent.cpp
@@ -1,6 +1,7 @@
 // Copyright Epic Games, Inc. All Rights Reserved.
 
 #include "RetargetComponent.h"
+#include "CapturePerformer.h"
 #include "Components/SkeletalMeshComponent.h"
 #include "LiveLinkInstance.h"
 #include "RetargetAnimInstance.h"
@@ -29,15 +30,14 @@ void URetargetComponent::SetCustomRetargetProfile(FRetargetProfile InProfile)
 			AnimInstance->UpdateCustomRetargetProfile(CustomRetargetProfile);
 		}
 	}
+	bIsDirty = true;
 }
 
 FRetargetProfile URetargetComponent::GetCustomRetargetProfile()
 {
 	/// Check for valid components and refs and then get the custom retarget profile struct
 	FRetargetProfile OutProfile;
-    TObjectPtr<USkeletalMeshComponent> ControlledMesh;
-    ControlledMesh = Cast<USkeletalMeshComponent>(ControlledSkeletalMeshComponent.GetComponent(GetOwner()));
-	if(ControlledMesh)
+	if(TObjectPtr<USkeletalMeshComponent> ControlledMesh = Cast<USkeletalMeshComponent>(ControlledSkeletalMeshComponent.GetComponent(GetOwner())))
 	{
 		if(ControlledMesh->GetAnimInstance())
 		{
@@ -51,21 +51,35 @@ FRetargetProfile URetargetComponent::GetCustomRetargetProfile()
 
 void URetargetComponent::SetForceMeshesFollowLeader(const bool bInBool)
 {
-	bForceOtherMeshesToFollowControlledMesh = bInBool;
-	SetForceOtherMeshesToFollowControlledMesh(bForceOtherMeshesToFollowControlledMesh);
+	SetForceOtherMeshesToFollowControlledMesh(bInBool);
+}
+
+void URetargetComponent::SetSourcePerformer(ACapturePerformer* InPerformer)
+{
+	if (IsValid(InPerformer))
+	{
+		SourcePerformer = InPerformer;
+		SetSourcePerformerMesh(InPerformer->GetSkeletalMeshComponent());
+	}
+	else
+	{
+		SourcePerformer = nullptr;
+		SetSourcePerformerMesh(Cast<USkeletalMeshComponent>(ControlledSkeletalMeshComponent.GetComponent(GetOwner())));
+	}
+	bIsDirty = true;
 }
 
 void URetargetComponent::SetSourcePerformerMesh(USkeletalMeshComponent* InPerformerMesh)
 {
 	SourceSkeletalMeshComponent.OverrideComponent = InPerformerMesh;
-	InitiateAnimation();
+	bIsDirty = true;
 }
 
 
 void URetargetComponent::SetControlledMesh(USkeletalMeshComponent* InControlledMesh)
 {
 	ControlledSkeletalMeshComponent.OverrideComponent = InControlledMesh;
-	InitiateAnimation();
+	bIsDirty = true;
 }
 
 void URetargetComponent::SetRetargetAsset(UIKRetargeter* InRetargetAsset)
@@ -73,8 +87,7 @@ void URetargetComponent::SetRetargetAsset(UIKRetargeter* InRetargetAsset)
 	RetargetAsset = InRetargetAsset; //It is possible to pass a nullptr to the AnimInstance 
 
 	/// Check for valid components and refs and then reint animation to reset pose
-	TObjectPtr<USkeletalMeshComponent> ControlledMesh = Cast<USkeletalMeshComponent>(ControlledSkeletalMeshComponent.GetComponent(GetOwner()));
-	if(ControlledMesh)
+	if(TObjectPtr<USkeletalMeshComponent> ControlledMesh = Cast<USkeletalMeshComponent>(ControlledSkeletalMeshComponent.GetComponent(GetOwner())))
 	{
 		if(ControlledMesh->GetAnimInstance())
 		{
@@ -90,13 +103,16 @@ void URetargetComponent::SetRetargetAsset(UIKRetargeter* InRetargetAsset)
 			ControlledMesh->InitAnim(true /*bForceReinit*/);
 		}
 	}
+	bIsDirty = true;
 }
 
 void URetargetComponent::OnRegister()
 {
 	Super::OnRegister();
 
-	const USkeletalMeshComponent* ControlledMesh = Cast<USkeletalMeshComponent> (ControlledSkeletalMeshComponent.GetComponent(GetOwner()));
+	SourceSkeletalMeshComponent.OverrideComponent = nullptr;
+
+	USkeletalMeshComponent* ControlledMesh = Cast<USkeletalMeshComponent> (ControlledSkeletalMeshComponent.GetComponent(GetOwner()));
 
 	if (ControlledMesh)
 	{
@@ -104,12 +120,18 @@ void URetargetComponent::OnRegister()
 
 		if(AnimInstance) //Only set these properties if we have a valid AnimInstance of the correct class.
 		{
-			SetForceOtherMeshesToFollowControlledMesh(bForceOtherMeshesToFollowControlledMesh);
+			SetForceMeshesFollowLeader(bForceOtherMeshesToFollowControlledMesh);
 			SetCustomRetargetProfile(CustomRetargetProfile);
 		}
 	}
 
-	bIsDirty = true;
+	if (!GetOwner()->HasAnyFlags(RF_ClassDefaultObject))
+	{
+		SetControlledMesh(ControlledMesh);
+		SetSourcePerformer(SourcePerformer.Get());
+		SetRetargetAsset(RetargetAsset);
+		SetForceMeshesFollowLeader(bForceOtherMeshesToFollowControlledMesh);
+	}
 }
 
 void URetargetComponent::TickComponent(float DeltaTime, ELevelTick TickType, FActorComponentTickFunction* ThisTickFunction)
@@ -153,29 +175,30 @@ void URetargetComponent::PostEditChangeProperty(FPropertyChangedEvent& PropertyC
 			bIsDirty=true;
 		}
 
-		Super::PostEditChangeProperty(PropertyChangedEvent);
 		if (PropertyChangedEvent.Property->GetFName() == GET_MEMBER_NAME_CHECKED(URetargetComponent, ControlledSkeletalMeshComponent))
 		{
 			bIsDirty=true;
 		}
 
-		Super::PostEditChangeProperty(PropertyChangedEvent);
 		if (PropertyChangedEvent.Property->GetFName() == GET_MEMBER_NAME_CHECKED(URetargetComponent, bForceOtherMeshesToFollowControlledMesh))
 		{
 			SetForceOtherMeshesToFollowControlledMesh(bForceOtherMeshesToFollowControlledMesh);
 		}
 
-		Super::PostEditChangeProperty(PropertyChangedEvent);
 		if (PropertyChangedEvent.Property->GetFName() == GET_MEMBER_NAME_CHECKED(URetargetComponent, RetargetAsset))
 		{
 			SetRetargetAsset(RetargetAsset);
-			bIsDirty=true;
 		}
 		
 		if (PropertyChangedEvent.Property->GetFName() == GET_MEMBER_NAME_CHECKED(URetargetComponent, CustomRetargetProfile))
 		{
 			SetCustomRetargetProfile(CustomRetargetProfile);
 		}
+
+		if (PropertyChangedEvent.Property->GetFName() == GET_MEMBER_NAME_CHECKED(URetargetComponent, SourcePerformer))
+		{
+			SetSourcePerformer(SourcePerformer.Get());
+		}
 	}
 }
 #endif
@@ -184,8 +207,7 @@ void URetargetComponent::SetForceOtherMeshesToFollowControlledMesh(bool bInBool)
 {
 	bForceOtherMeshesToFollowControlledMesh = bInBool;
 	
-	USkeletalMeshComponent* ControlledMesh = Cast<USkeletalMeshComponent> (ControlledSkeletalMeshComponent.GetComponent
-(GetOwner()));
+	USkeletalMeshComponent* ControlledMesh = Cast<USkeletalMeshComponent> (ControlledSkeletalMeshComponent.GetComponent(GetOwner()));
 	const USkeletalMeshComponent* SourceMesh = Cast<USkeletalMeshComponent>(SourceSkeletalMeshComponent.GetComponent(GetOwner()));
 	
 	TArray<USkeletalMeshComponent*> OwnerSkeletalMeshComponents;
@@ -230,6 +252,7 @@ void URetargetComponent::SetForceOtherMeshesToFollowControlledMesh(bool bInBool)
 			}
 		}
 	}
+	bIsDirty = true;
 }
 
 void URetargetComponent::InitiateAnimation()
diff --git a/Engine/Plugins/Animation/PerformanceCaptureCore/Source/PerformanceCaptureCore/Public/CaptureCharacter.h b/Engine/Plugins/Animation/PerformanceCaptureCore/Source/PerformanceCaptureCore/Public/CaptureCharacter.h
--- a/Engine/Plugins/Animation/PerformanceCaptureCore/Source/PerformanceCaptureCore/Public/CaptureCharacter.h
+++ b/Engine/Plugins/Animation/PerformanceCaptureCore/Source/PerformanceCaptureCore/Public/CaptureCharacter.h
@@ -22,24 +22,29 @@ public:
 	// Sets default values for this actor's properties
 	UE_API ACaptureCharacter();
 
+#if WITH_EDITORONLY_DATA
 	/**
 	* CapturePerformer Actor that will be the source for retargeting.
 	*/
-	UPROPERTY(BlueprintReadOnly, EditAnywhere, Category = "Performance Capture|Character")
-	TSoftObjectPtr<ACapturePerformer> SourcePerformer;
+	UE_DEPRECATED(5.7, "Use URetargetComponent::SourcePerformer instead.")
+	UPROPERTY(meta = (DeprecatedProperty, DeprecationMessage = "Use URetargetComponent::SourcePerformer instead."))
+	TSoftObjectPtr<ACapturePerformer> SourcePerformer_DEPRECATED;
 
 	/**
 	* The IKRetarget Asset to use for retargeting between the SourcePerformer and this Character.
 	*/
-	UPROPERTY(BlueprintReadOnly, EditAnywhere, Category = "Performance Capture|Character")
-	TObjectPtr<UIKRetargeter> RetargetAsset;
+	UE_DEPRECATED(5.7, "Use URetargetComponent::RetargetAsset instead.")
+	UPROPERTY(meta = (DeprecatedProperty, DeprecationMessage = "Use URetargetComponent::RetargetAsset instead."))
+	TObjectPtr<UIKRetargeter> RetargetAsset_DEPRECATED;
 
 	/**
 	* Force all skeletal meshes to use the root SkeletalMesh as their Leader. Default = True.
 	*/
-	UPROPERTY(BlueprintReadOnly, EditAnywhere, Category = "Performance Capture|Character")
-	bool bForceAllSkeletalMeshesToFollowLeader = true;
-	
+	UE_DEPRECATED(5.7, "Use URetargetComponent::bForceOtherMeshesToFollowControlledMesh instead.")
+	UPROPERTY(meta = (DeprecatedProperty, DeprecationMessage = "Use URetargetComponent::bForceOtherMeshesToFollowControlledMesh instead."))
+	bool bForceAllSkeletalMeshesToFollowLeader_DEPRECATED = true;
+#endif // WITH_EDITORONLY_DATA
+
 	/**
 	* Set a Custom Retarget Profile.
 	* @param InProfile New Retarget Profile.
@@ -75,11 +80,16 @@ public:
 	UFUNCTION(BlueprintCallable, Category = "Performance Capture|Character")
 	UE_API void SetForceAllSkeletalMeshesToFollowLeader (bool InFollowLeader);
 
+	/**
+	* Return the retarget component
+	*/
+	UFUNCTION(BlueprintPure, Category = "Performance Capture|Character")
+	UE_API URetargetComponent* GetRetargetComponent() const;
+
 protected:
-	UE_API virtual void PostRegisterAllComponents() override;
-#if WITH_EDITOR
-	UE_API virtual void PostEditChangeProperty(FPropertyChangedEvent& PropertyChangedEvent) override;
-#endif
+	// UObject interface
+	virtual void PostLoad() override;
+	// End of UObject interface
 
 private:
 	// retarget component is private and hidden from the editor UI
diff --git a/Engine/Plugins/Animation/PerformanceCaptureCore/Source/PerformanceCaptureCore/Public/RetargetComponent.h b/Engine/Plugins/Animation/PerformanceCaptureCore/Source/PerformanceCaptureCore/Public/RetargetComponent.h
--- a/Engine/Plugins/Animation/PerformanceCaptureCore/Source/PerformanceCaptureCore/Public/RetargetComponent.h
+++ b/Engine/Plugins/Animation/PerformanceCaptureCore/Source/PerformanceCaptureCore/Public/RetargetComponent.h
@@ -10,7 +10,7 @@
 
 #define UE_API PERFORMANCECAPTURECORE_API
 
-
+class ACapturePerformer;
 class UIKRetargeter;
 
 UCLASS(MinimalAPI, BlueprintType, ClassGroup=("Performance Capture"), meta=(BlueprintSpawnableComponent), DisplayName = "Retarget Component")
@@ -25,7 +25,8 @@ public:
 	/**
 	* Skeletal Mesh component that will be the source for retargeting. Can be on the Owner Actor or another Actor in the same level.
 	*/
-	UPROPERTY(BlueprintReadOnly, EditInstanceOnly, Category = "Performance Capture",  meta = (UseComponentPicker, AllowAnyActor, AllowedClasses = "/Script/Engine.SkeletalMeshComponent"))
+	UPROPERTY(BlueprintReadOnly, EditInstanceOnly, Category = "Performance Capture",  
+		meta = (UseComponentPicker, AllowAnyActor, AllowedClasses = "/Script/Engine.SkeletalMeshComponent", EditCondition = "IsSourcePerformerNull()", EditConditionHides))
 	FComponentReference SourceSkeletalMeshComponent ;
 
 	/**
@@ -51,6 +52,12 @@ public:
 	*/
 	UPROPERTY(BlueprintReadOnly, Interp, EditAnywhere, Category = "Performance Capture")
 	FRetargetProfile CustomRetargetProfile;
+
+	/**
+	* CapturePerformer Actor that will be the source for retargeting.
+	*/
+	UPROPERTY(BlueprintReadOnly, EditAnywhere, Category = "Performance Capture|Character")
+	TSoftObjectPtr<ACapturePerformer> SourcePerformer;
 	
 	/**
 	* Set the Source Performer Mesh.
@@ -93,6 +100,13 @@ public:
 	*/
 	UFUNCTION(BlueprintCallable, Category = "Performance Capture|Retarget")
 	UE_API void SetForceMeshesFollowLeader(bool bInBool);
+
+	/**
+	* Set the Source CapturePerformer Actor.
+	* @param InPerformer New Performer.
+	*/
+	UFUNCTION(BlueprintCallable, Category = "Performance Capture|Character")
+	UE_API void SetSourcePerformer(ACapturePerformer* InPerformer);
 	
 protected:
 	// Called when the game starts
@@ -103,6 +117,9 @@ protected:
 #if WITH_EDITOR
 	UE_API virtual void PostEditChangeProperty(FPropertyChangedEvent& PropertyChangedEvent) override;
 #endif
+
+	UFUNCTION()
+	inline bool IsSourcePerformerNull() const { return SourcePerformer.IsNull(); }
 	
 private:
 	bool bIsDirty = true;
diff --git a/Engine/Plugins/VirtualProduction/PerformanceCaptureWorkflow/Config/BasePerformanceCaptureWorkflow.ini b/Engine/Plugins/VirtualProduction/PerformanceCaptureWorkflow/Config/BasePerformanceCaptureWorkflow.ini
--- a/Engine/Plugins/VirtualProduction/PerformanceCaptureWorkflow/Config/BasePerformanceCaptureWorkflow.ini
+++ b/Engine/Plugins/VirtualProduction/PerformanceCaptureWorkflow/Config/BasePerformanceCaptureWorkflow.ini
@@ -7,6 +7,8 @@ ViewModelClass=/PerformanceCaptureWorkflow/Core/VM_PCapWorkflow.VM_PCapWorkflow_
 MocapManagerUI=/PerformanceCaptureWorkflow/Core/EW_MocapManager.EW_MocapManager
 DefaultSessionTemplate=/PerformanceCaptureWorkflow/Core/DefaultSessionTemplate.DefaultSessionTemplate
 DatabaseHelperClass=/PerformanceCaptureWorkflow/Database/EU_InternalDatabaseHelper.EU_InternalDatabaseHelper_C
+AllowedCaptureCharacterActorClasses=/Script/PerformanceCaptureCore.CaptureCharacter
+DisallowedCaptureCharacterActorClasses=/Script/PerformanceCaptureCore.CapturePerformer
 
 [CoreRedirects]
 +ClassRedirects=(OldName="/Script/PerformanceCaptureWorkflow.PCap_DataAsset",NewName="/Script/PerformanceCaptureWorkflow.PCapDataAsset")
diff --git a/Engine/Plugins/VirtualProduction/PerformanceCaptureWorkflow/Source/PerformanceCaptureWorkflow/Private/PCapActorFactory.cpp b/Engine/Plugins/VirtualProduction/PerformanceCaptureWorkflow/Source/PerformanceCaptureWorkflow/Private/PCapActorFactory.cpp
--- a/Engine/Plugins/VirtualProduction/PerformanceCaptureWorkflow/Source/PerformanceCaptureWorkflow/Private/PCapActorFactory.cpp
+++ b/Engine/Plugins/VirtualProduction/PerformanceCaptureWorkflow/Source/PerformanceCaptureWorkflow/Private/PCapActorFactory.cpp
@@ -15,9 +15,59 @@
 #include "SubobjectDataSubsystem.h"
 #include "Elements/Actor/ActorElementData.h"
 #include "Elements/Framework/EngineElementsLibrary.h"
-#include "Subsystems/EditorActorSubsystem.h"
+#include "Logging/MessageLog.h"
+#include "RetargetComponent.h"
 #include "Subsystems/PlacementSubsystem.h"
 
+#define LOCTEXT_NAMESPACE "PCapActorFactory"
+
+namespace UE::PCapActorFactory
+{
+
+	/** Find and return an existing Component of type TComponent or add and return a new one if not. */
+	UActorComponent* AddComponent(TSubclassOf<UActorComponent> ComponentClass, AActor* OnActor, bool bLogMessageOnError = true)
+	{
+		if(ComponentClass == nullptr)
+		{
+			return nullptr;
+		}
+
+		USubobjectDataSubsystem* SubobjectSubsystem = GEngine->GetEngineSubsystem<USubobjectDataSubsystem>();
+		TArray<FSubobjectDataHandle> SubObjectHandles;
+		SubobjectSubsystem->GatherSubobjectData(OnActor, SubObjectHandles);
+		FSubobjectDataHandle RootHandle = SubObjectHandles[0]; //The actor's root component is always the first handle in the array
+		FAddNewSubobjectParams NewObjectParams;
+		NewObjectParams.ParentHandle = RootHandle;
+		NewObjectParams.NewClass = ComponentClass;
+		NewObjectParams.bConformTransformToParent = true;
+
+		FText OutFailText;
+		FSubobjectDataHandle NewComponentHandle = SubobjectSubsystem->AddNewSubobject(NewObjectParams, OutFailText);
+
+		UActorComponent* Component = nullptr;
+		if(NewComponentHandle.IsValid())
+		{
+			// GetObject returns a const* so need to cast that away to make this usable.
+			Component = const_cast<UActorComponent*>(NewComponentHandle.GetData()->GetObject<UActorComponent>());
+		}
+		else if(bLogMessageOnError)
+		{
+			FFormatNamedArguments FormatArguments;
+			FormatArguments.Add(TEXT("Actor"), FText::FromString(OnActor->GetName()));
+			FormatArguments.Add(TEXT("ComponentType"), FText::FromString(ComponentClass->GetName()));
+			FormatArguments.Add(TEXT("ErrorMessage"), OutFailText);
+
+			FText Message = FText::Format(
+				LOCTEXT("PCapActorFactoryAddComponentFail", "Unable to add Component {ComponentType} to {Actor}, error: {ErrorMessage}"),
+				FormatArguments);
+
+			FMessageLog(FPerformanceCaptureModule::MessageLogName).Error(Message);
+		}
+
+		return Component;
+	}
+
+} // namespace UE::PCapActorFactory
 
 /*------------------------------------------------------------------------------
 	Custom Actor Factory for Character Data Asset implementation.
@@ -59,53 +109,30 @@ bool UPCapCharacterActorFactory::CanCreateActorFrom(const FAssetData& AssetData,
 void UPCapCharacterActorFactory::PostSpawnActor(UObject* Asset, AActor* NewActor)
 {
 	Super::PostSpawnActor(Asset, NewActor);
-	
-	if(Cast<ACaptureCharacter>(NewActor))
+
+	const UPCapCharacterDataAsset* CharacterAsset = Cast<UPCapCharacterDataAsset>(Asset);
+	ASkeletalMeshActor* AsSkeletalMeshActor = Cast<ASkeletalMeshActor>(NewActor);
+
+	if(AsSkeletalMeshActor != nullptr && CharacterAsset != nullptr && !AsSkeletalMeshActor->IsA<ACapturePerformer>())
 	{
-		ACaptureCharacter* NewCharacter = Cast<ACaptureCharacter>(NewActor);
-		
-		if(UPCapCharacterDataAsset* CharacterAsset = Cast<UPCapCharacterDataAsset>(Asset))
+		TObjectPtr<URetargetComponent> RetargetComponent = AsSkeletalMeshActor->GetComponentByClass<URetargetComponent>();
+
+		if(RetargetComponent == nullptr)
+		{
+			RetargetComponent = Cast<URetargetComponent>(UE::PCapActorFactory::AddComponent(URetargetComponent::StaticClass(), AsSkeletalMeshActor));
+		}
+
+		if(RetargetComponent != nullptr)
 		{
 			USkeletalMesh* CharacterMesh = CharacterAsset->SkeletalMesh.LoadSynchronous();
-			NewCharacter->GetSkeletalMeshComponent()->SetSkeletalMeshAsset(CharacterMesh);
-			NewCharacter->SetRetargetAsset(Cast<UIKRetargeter>(CharacterAsset->Retargeter.LoadSynchronous()));
-			
+			AsSkeletalMeshActor->GetSkeletalMeshComponent()->SetSkeletalMeshAsset(CharacterMesh);
+			RetargetComponent->SetRetargetAsset(Cast<UIKRetargeter>(CharacterAsset->Retargeter.LoadSynchronous()));
+
 			if(SpawnedPerformerHandle.IsSet()) //Entry here is only once per placing action.
 			{
 				ACapturePerformer* SourcePerformer = Cast<ACapturePerformer>(ActorElementDataUtil::GetActorFromHandle(SpawnedPerformerHandle));
-				NewCharacter->SourcePerformer = SourcePerformer;
-				NewCharacter->GetSkeletalMeshComponent()->SetReceivesDecals(false);
-
-				if(!CharacterAsset->AdditionalMeshes.IsEmpty())
-				{
-					if (NewActor->HasAnyFlags(RF_Transactional) == true) //Do not spawn any subcomponents on the transient copy of the new actor
-					{
-						USubobjectDataSubsystem* SubSystem = GEngine->GetEngineSubsystem<USubobjectDataSubsystem>();
-
-						for (TSoftObjectPtr<USkeletalMesh> Mesh : CharacterAsset->AdditionalMeshes)
-						{
-							TArray<FSubobjectDataHandle> SubObjectHandles;
-							SubSystem->GatherSubobjectData(NewActor, SubObjectHandles);
-							FSubobjectDataHandle RootHandle = SubObjectHandles[0]; //The actor's root component is always the first handle in the array
-							FAddNewSubobjectParams NewObjectParams;
-							NewObjectParams.ParentHandle = RootHandle;
-							NewObjectParams.NewClass = USkeletalMeshComponent::StaticClass();
-							NewObjectParams.bConformTransformToParent = true;
-
-							FText FailureReason;
-							FSubobjectDataHandle NewSkelMeshComponentHandle = SubSystem->AddNewSubobject(NewObjectParams, FailureReason);
-					
-							TArray<USkeletalMeshComponent*> SkeletalMeshComponents;
-							NewCharacter->GetComponents(SkeletalMeshComponents);
-					
-							if(TObjectPtr<USkeletalMeshComponent> const NewSkeletalMeshComponent = SkeletalMeshComponents.Last())
-							{
-								NewSkeletalMeshComponent->SetSkeletalMeshAsset(Mesh.LoadSynchronous());
-								NewSkeletalMeshComponent->SetReceivesDecals(false);
-							}
-						}
-					}
-				}
+				RetargetComponent->SourcePerformer = SourcePerformer;
+				AsSkeletalMeshActor->GetSkeletalMeshComponent()->SetReceivesDecals(false);
 			}
 		}
 	}
@@ -362,3 +389,5 @@ void UPCapPropActorFactory::PostSpawnActor(UObject* Asset, AActor* NewActor)
 	}
 }
 
+#undef LOCTEXT_NAMESPACE
+
diff --git a/Engine/Plugins/VirtualProduction/PerformanceCaptureWorkflow/Source/PerformanceCaptureWorkflow/Private/PCapDatabase.cpp b/Engine/Plugins/VirtualProduction/PerformanceCaptureWorkflow/Source/PerformanceCaptureWorkflow/Private/PCapDatabase.cpp
--- a/Engine/Plugins/VirtualProduction/PerformanceCaptureWorkflow/Source/PerformanceCaptureWorkflow/Private/PCapDatabase.cpp
+++ b/Engine/Plugins/VirtualProduction/PerformanceCaptureWorkflow/Source/PerformanceCaptureWorkflow/Private/PCapDatabase.cpp
@@ -1,8 +1,10 @@
 // Copyright Epic Games, Inc. All Rights Reserved.
 
 #include "PCapDatabase.h"
+#include "Algo/Transform.h"
 #include "CapturePerformer.h"
 #include "CaptureCharacter.h"
+#include "PCapSettings.h"
 #include "PerformanceCapture.h"
 #include "Misc/Guid.h"
 
@@ -63,6 +65,32 @@ UPCapCharacterDataAsset::UPCapCharacterDataAsset()
 {
 }
 
+TArray<UClass*> UPCapCharacterDataAsset::GetAllowedCaptureCharacterActorClasses()
+{
+	UPerformanceCaptureSettings* Settings = UPerformanceCaptureSettings::GetPerformanceCaptureSettings();
+	TArray<UClass*> AllowedClasses;
+	AllowedClasses.Reserve(Settings->AllowedCaptureCharacterActorClasses.Num());
+
+	Algo::TransformIf(Settings->AllowedCaptureCharacterActorClasses, AllowedClasses,
+		[](const TSoftClassPtr<ASkeletalMeshActor>& InSoftClass) { return !InSoftClass.IsNull(); },
+		[](const TSoftClassPtr<ASkeletalMeshActor>& InSoftClass) -> UClass* { return InSoftClass.LoadSynchronous(); }
+	);
+	return AllowedClasses;
+}
+
+TArray<UClass*> UPCapCharacterDataAsset::GetDisallowedCaptureCharacterActorClasses()
+{
+	UPerformanceCaptureSettings* Settings = UPerformanceCaptureSettings::GetPerformanceCaptureSettings();
+	TArray<UClass*> DisallowedClasses;
+	DisallowedClasses.Reserve(Settings->DisallowedCaptureCharacterActorClasses.Num());
+
+	Algo::TransformIf(Settings->DisallowedCaptureCharacterActorClasses, DisallowedClasses,
+		[](const TSoftClassPtr<ASkeletalMeshActor>& InSoftClass) { return !InSoftClass.IsNull(); },
+		[](const TSoftClassPtr<ASkeletalMeshActor>& InSoftClass) -> UClass* { return InSoftClass.LoadSynchronous(); }
+	);
+	return DisallowedClasses;
+}
+
 UPCapPropDataAsset::UPCapPropDataAsset()
 	: UPCapDataAsset()
 {
diff --git a/Engine/Plugins/VirtualProduction/PerformanceCaptureWorkflow/Source/PerformanceCaptureWorkflow/Private/PCapDatabase.h b/Engine/Plugins/VirtualProduction/PerformanceCaptureWorkflow/Source/PerformanceCaptureWorkflow/Private/PCapDatabase.h
--- a/Engine/Plugins/VirtualProduction/PerformanceCaptureWorkflow/Source/PerformanceCaptureWorkflow/Private/PCapDatabase.h
+++ b/Engine/Plugins/VirtualProduction/PerformanceCaptureWorkflow/Source/PerformanceCaptureWorkflow/Private/PCapDatabase.h
@@ -371,8 +371,9 @@ public:
 	TSoftObjectPtr<UPCapPerformerDataAsset> SourcePerformerAsset;
 
 	/** Character actor Class */
-	UPROPERTY(BlueprintReadWrite, EditAnywhere, Category="Performance Capture Character")
-	TSoftClassPtr<ACaptureCharacter>  CaptureCharacterClass;
+	UPROPERTY(BlueprintReadWrite, EditAnywhere, Category = "Performance Capture Character",
+		meta = (GetAllowedClasses = "GetAllowedCaptureCharacterActorClasses", GetDisallowedClasses = "GetDisallowedCaptureCharacterActorClasses"))
+	TSoftClassPtr<ASkeletalMeshActor>  CaptureCharacterClass;
 
 	/** Main Skeletal Mesh Asset. This will be the root component of any characters spawned.
 	 * A character will not be spawned if this asset reference is null or not valid.
@@ -388,9 +389,19 @@ public:
 	UPROPERTY(BlueprintReadWrite, EditAnywhere, Category="Performance Capture Character")
 	TSoftObjectPtr<UIKRetargeter> Retargeter;
 
+#if WITH_EDITORONLY_DATA
 	/** Array of additional skeleltal mesh assets that will be spawned and parented to the root skeleteal mesh asset*/
-	UPROPERTY(BlueprintReadWrite, EditAnywhere, Category="Performance Capture Character", meta=(DisplayName="Additional Meshes", ToolTip="Additional meshes that will be added as child components"))
-	TArray<TSoftObjectPtr<USkeletalMesh>> AdditionalMeshes;
+	UE_DEPRECATED(5.7, "AdditionalMeshes is no longer used. Create a SkeletalMeshActor class with the required meshes and components instead.")
+	UPROPERTY(meta=(DeprecatedProperty, DeprecationMessage = "Create a SkeletalMeshActor class with the required meshes and components instead."))
+	TArray<TSoftObjectPtr<USkeletalMesh>> AdditionalMeshes_DEPRECATED;
+#endif // WITH_EDITORONLY_DATA
+
+private:
+	UFUNCTION()
+	static TArray<UClass*> GetAllowedCaptureCharacterActorClasses();
+
+	UFUNCTION()
+	static TArray<UClass*> GetDisallowedCaptureCharacterActorClasses();
 };
 
 /**  Prop DataAsset*/
diff --git a/Engine/Plugins/VirtualProduction/PerformanceCaptureWorkflow/Source/PerformanceCaptureWorkflow/Private/PCapSettings.h b/Engine/Plugins/VirtualProduction/PerformanceCaptureWorkflow/Source/PerformanceCaptureWorkflow/Private/PCapSettings.h
--- a/Engine/Plugins/VirtualProduction/PerformanceCaptureWorkflow/Source/PerformanceCaptureWorkflow/Private/PCapSettings.h
+++ b/Engine/Plugins/VirtualProduction/PerformanceCaptureWorkflow/Source/PerformanceCaptureWorkflow/Private/PCapSettings.h
@@ -2,6 +2,7 @@
 
 #pragma once
 
+#include "Animation/SkeletalMeshActor.h"
 #include "CoreMinimal.h"
 #include "EditorUtilityWidgetBlueprint.h"
 #include "PCapStageRoot.h"
@@ -47,6 +48,14 @@ public:
 	UPROPERTY(BlueprintReadOnly, EditDefaultsOnly, Config, NoClear, Category="Actors", DisplayName= "Default Performer Mesh")
 	TSoftObjectPtr<USkeletalMesh> DefaultPerformerSkelMesh;
 
+	/** Skeletal Mesh Actor classes that are explicitly allowed for use with Mocap Manager. */
+	UPROPERTY(BlueprintReadOnly, EditDefaultsOnly, Config, NoClear, Category = "Actors")
+	TArray<TSoftClassPtr<ASkeletalMeshActor>> AllowedCaptureCharacterActorClasses;
+
+	/** Skeletal Mesh Actor classes that are explicitly disallowed for use with Mocap Manager. */
+	UPROPERTY(BlueprintReadOnly, EditDefaultsOnly, Config, NoClear, Category = "Actors")
+	TArray<TSoftClassPtr<ASkeletalMeshActor>> DisallowedCaptureCharacterActorClasses;
+
 	/** Blueprint Viewmodel class. If you change class you will need to restart the editor to pick up the change. */
 	UPROPERTY(EditAnywhere, BlueprintReadOnly, Config, NoClear, Category="UI", meta = (MetaClass = "/Script/ModelViewViewModel.MVVMViewModelBase",ConfigRestartRequired=true))
 	TSoftClassPtr<UMVVMViewModelBase> ViewModelClass;
@@ -96,6 +105,5 @@ public:
 	FOnPCapSettingsChanged OnPCapSettingsChanged;
 	
 	virtual void PostEditChangeProperty(FPropertyChangedEvent& PropertyChangedEvent) override;
-	
 #endif
 };
diff --git a/Engine/Plugins/VirtualProduction/PerformanceCaptureWorkflow/Source/PerformanceCaptureWorkflow/Private/PerformanceCapture.cpp b/Engine/Plugins/VirtualProduction/PerformanceCaptureWorkflow/Source/PerformanceCaptureWorkflow/Private/PerformanceCapture.cpp
--- a/Engine/Plugins/VirtualProduction/PerformanceCaptureWorkflow/Source/PerformanceCaptureWorkflow/Private/PerformanceCapture.cpp
+++ b/Engine/Plugins/VirtualProduction/PerformanceCaptureWorkflow/Source/PerformanceCaptureWorkflow/Private/PerformanceCapture.cpp
@@ -7,6 +7,7 @@
 #include "PerformanceCaptureCommands.h"
 #include "ISettingsModule.h"
 #include "LevelEditorOutlinerSettings.h"
+#include "MessageLogModule.h"
 #include "PCapAssetDefinition.h"
 #include "PCapDatabase.h"
 #include "PCapSettings.h"
@@ -49,6 +50,9 @@ void FPerformanceCaptureModule::StartupModule()
 		.SetIcon(FSlateIcon(FPerformanceCaptureStyle::GetStyleSetName(), "PerformanceCapture.MocapManagerTabIcon", "PerformanceCapture.MocapManagerTabIcon.Small")
 		);
 
+	FMessageLogModule& MessageLogModule = FModuleManager::LoadModuleChecked<FMessageLogModule>("MessageLog");
+	MessageLogModule.RegisterLogListing(MessageLogName, LOCTEXT("MessageLogPerformanceCaptureName", "PerformanceCapture"));
+
 }
 
 void FPerformanceCaptureModule::ShutdownModule()
diff --git a/Engine/Plugins/VirtualProduction/PerformanceCaptureWorkflow/Source/PerformanceCaptureWorkflow/Public/PerformanceCapture.h b/Engine/Plugins/VirtualProduction/PerformanceCaptureWorkflow/Source/PerformanceCaptureWorkflow/Public/PerformanceCapture.h
--- a/Engine/Plugins/VirtualProduction/PerformanceCaptureWorkflow/Source/PerformanceCaptureWorkflow/Public/PerformanceCapture.h
+++ b/Engine/Plugins/VirtualProduction/PerformanceCaptureWorkflow/Source/PerformanceCaptureWorkflow/Public/PerformanceCapture.h
@@ -20,6 +20,9 @@ public:
 	/** This function will be bound to Command (by default it will bring up plugin window) */
 	void PluginButtonClicked();
 
+	/** Message Log Name for the PerformanceCapture module. */
+	inline static const FLazyName MessageLogName = TEXT("PerformanceCapture");
+
 private:
 
 	void RegisterMenus();
