MocapManager: Add Details customization for ACaptureCharacter to display its RetargetComponent.

diff --git a/Engine/Plugins/Animation/PerformanceCaptureCore/PerformanceCaptureCore.uplugin b/Engine/Plugins/Animation/PerformanceCaptureCore/PerformanceCaptureCore.uplugin
--- a/Engine/Plugins/Animation/PerformanceCaptureCore/PerformanceCaptureCore.uplugin
+++ b/Engine/Plugins/Animation/PerformanceCaptureCore/PerformanceCaptureCore.uplugin
@@ -18,6 +18,11 @@
 			"Name": "PerformanceCaptureCore",
 			"Type": "Runtime",
 			"LoadingPhase": "Default"
+		},
+		{
+			"Name": "PerformanceCaptureCoreEditor",
+			"Type": "Editor",
+			"LoadingPhase": "Default"
 		}
 	],
 	"Plugins": [
diff --git a/Engine/Plugins/Animation/PerformanceCaptureCore/Source/PerformanceCaptureCoreEditor/PerformanceCaptureCoreEditor.Build.cs b/Engine/Plugins/Animation/PerformanceCaptureCore/Source/PerformanceCaptureCoreEditor/PerformanceCaptureCoreEditor.Build.cs
new file mode 100644
--- /dev/null
+++ b/Engine/Plugins/Animation/PerformanceCaptureCore/Source/PerformanceCaptureCoreEditor/PerformanceCaptureCoreEditor.Build.cs
@@ -0,0 +1,30 @@
+// Copyright Epic Games, Inc. All Rights Reserved.
+
+using UnrealBuildTool;
+
+public class PerformanceCaptureCoreEditor : ModuleRules
+{
+	public PerformanceCaptureCoreEditor(ReadOnlyTargetRules Target) : base(Target)
+	{
+		PCHUsage = ModuleRules.PCHUsageMode.UseExplicitOrSharedPCHs;
+		
+		PublicDependencyModuleNames.AddRange(
+			new string[]
+			{
+				"Core",
+			}
+			);
+		
+		PrivateDependencyModuleNames.AddRange(
+			new string[]
+			{
+				"CoreUObject",
+				"Engine",
+				"IKRig",
+				"LiveLinkAnimationCore",
+				"LiveLinkInterface",
+				"PerformanceCaptureCore"
+			}
+			);
+	}
+}
diff --git a/Engine/Plugins/Animation/PerformanceCaptureCore/Source/PerformanceCaptureCoreEditor/Private/CaptureCharacterCustomization.cpp b/Engine/Plugins/Animation/PerformanceCaptureCore/Source/PerformanceCaptureCoreEditor/Private/CaptureCharacterCustomization.cpp
new file mode 100644
--- /dev/null
+++ b/Engine/Plugins/Animation/PerformanceCaptureCore/Source/PerformanceCaptureCoreEditor/Private/CaptureCharacterCustomization.cpp
@@ -0,0 +1,59 @@
+// Copyright Epic Games, Inc. All Rights Reserved.
+
+#include "CaptureCharacterCustomization.h"
+
+#include "CaptureCharacter.h"
+#include "DetailLayoutBuilder.h"
+#include "DetailWidgetRow.h"
+#include "Framework/MultiBox/MultiBoxBuilder.h"
+#include "IDetailGroup.h"
+#include "PropertyCustomizationHelpers.h"
+#include "RetargetComponent.h"
+
+#define LOCTEXT_NAMESPACE "CaptureCharacterCustomization"
+
+TSharedRef<IDetailCustomization> FCaptureCharacterCustomization::MakeInstance()
+{
+	return MakeShared<FCaptureCharacterCustomization>();
+}
+
+void FCaptureCharacterCustomization::CustomizeDetails(IDetailLayoutBuilder& DetailBuilder)
+{
+	TArray<TWeakObjectPtr<UObject>> CustomizedObjects;
+	DetailBuilder.GetObjectsBeingCustomized(CustomizedObjects);
+
+	// Ensure that we are only customizing one object
+	if (CustomizedObjects.Num() != 1)
+	{
+		return;
+	}
+
+	CustomizedCharacter = Cast<ACaptureCharacter>(CustomizedObjects[0]);
+
+	CustomizePerformanceCaptureCategory(DetailBuilder);
+}
+
+void FCaptureCharacterCustomization::CustomizePerformanceCaptureCategory(IDetailLayoutBuilder& DetailBuilder)
+{
+	// Customize the PerformanceCapture category for the Retarget Component property.
+	TObjectPtr<URetargetComponent> CustomizedRetargetComponent = CustomizedCharacter->GetRetargetComponent();
+
+	if (CustomizedRetargetComponent == nullptr)
+	{
+		return;
+	}
+
+	IDetailCategoryBuilder& PCapCategory = DetailBuilder.EditCategory("Performance Capture", LOCTEXT("PerformanceCaptureCategoryName", "Performance Capture"));
+
+	const TArray<UObject*> CustomizedRetargetComponentObjects = { CustomizedRetargetComponent };
+	PCapCategory.AddExternalObjectProperty(CustomizedRetargetComponentObjects, GET_MEMBER_NAME_CHECKED(URetargetComponent, ControlledSkeletalMeshComponent));
+	PCapCategory.AddExternalObjectProperty(CustomizedRetargetComponentObjects, GET_MEMBER_NAME_CHECKED(URetargetComponent, bForceOtherMeshesToFollowControlledMesh));
+	PCapCategory.AddExternalObjectProperty(CustomizedRetargetComponentObjects, GET_MEMBER_NAME_CHECKED(URetargetComponent, RetargetAsset));
+	PCapCategory.AddExternalObjectProperty(CustomizedRetargetComponentObjects, GET_MEMBER_NAME_CHECKED(URetargetComponent, CustomRetargetProfile));
+	PCapCategory.AddExternalObjectProperty(CustomizedRetargetComponentObjects, GET_MEMBER_NAME_CHECKED(URetargetComponent, SourceSkeletalMeshComponent));
+
+	IDetailGroup& CharacterGroup = PCapCategory.AddGroup("Character", LOCTEXT("PerformanceCaptureCharacterGroupName", "Character"));
+	CharacterGroup.AddExternalObjectProperty(CustomizedRetargetComponentObjects, GET_MEMBER_NAME_CHECKED(URetargetComponent, SourcePerformer), EPropertyLocation::Default, FAddPropertyParams());
+}
+
+#undef LOCTEXT_NAMESPACE
diff --git a/Engine/Plugins/Animation/PerformanceCaptureCore/Source/PerformanceCaptureCoreEditor/Private/CaptureCharacterCustomization.h b/Engine/Plugins/Animation/PerformanceCaptureCore/Source/PerformanceCaptureCoreEditor/Private/CaptureCharacterCustomization.h
new file mode 100644
--- /dev/null
+++ b/Engine/Plugins/Animation/PerformanceCaptureCore/Source/PerformanceCaptureCoreEditor/Private/CaptureCharacterCustomization.h
@@ -0,0 +1,31 @@
+// Copyright Epic Games, Inc. All Rights Reserved.
+
+#pragma once
+
+#include "IDetailCustomization.h"
+
+#include "Styling/SlateTypes.h"
+#include "UObject/TemplateString.h"
+
+class ACaptureCharacter;
+
+/**
+ * Detail customization for ACaptureCharacter
+ */
+class FCaptureCharacterCustomization : public IDetailCustomization
+{
+public:
+	/** Makes a new instance of this detail layout class for a specific detail view requesting it */
+	static TSharedRef<IDetailCustomization> MakeInstance();
+
+	//~ Begin IDetailCustomization interface */
+	virtual void CustomizeDetails(IDetailLayoutBuilder& DetailBuilder) override;
+	//~ End IDetailCustomization interface */
+
+private:
+	void CustomizePerformanceCaptureCategory(IDetailLayoutBuilder& DetailBuilder);
+
+private:
+	/** The CaptureCharacter being customized */
+	ACaptureCharacter* CustomizedCharacter;
+};
diff --git a/Engine/Plugins/Animation/PerformanceCaptureCore/Source/PerformanceCaptureCoreEditor/Private/PerformanceCaptureCoreEditor.cpp b/Engine/Plugins/Animation/PerformanceCaptureCore/Source/PerformanceCaptureCoreEditor/Private/PerformanceCaptureCoreEditor.cpp
new file mode 100644
--- /dev/null
+++ b/Engine/Plugins/Animation/PerformanceCaptureCore/Source/PerformanceCaptureCoreEditor/Private/PerformanceCaptureCoreEditor.cpp
@@ -0,0 +1,32 @@
+// Copyright Epic Games, Inc. All Rights Reserved.
+
+#include "CaptureCharacter.h"
+#include "CaptureCharacterCustomization.h"
+#include "Modules/ModuleInterface.h"
+#include "Modules/ModuleManager.h"
+#include "PropertyEditorModule.h"
+
+#define LOCTEXT_NAMESPACE "FPerformanceCaptureCoreEditorModule"
+
+/** Editor module for FPerformanceCaptureCoreEditor */
+class FPerformanceCaptureCoreEditorModule : public IModuleInterface
+{
+	void StartupModule() override
+	{
+		FPropertyEditorModule& PropertyEditorModule = FModuleManager::Get().LoadModuleChecked<FPropertyEditorModule>("PropertyEditor");
+		PropertyEditorModule.RegisterCustomClassLayout(
+			ACaptureCharacter::StaticClass()->GetFName(),
+			FOnGetDetailCustomizationInstance::CreateStatic(&FCaptureCharacterCustomization::MakeInstance)
+		);
+	}
+
+	void ShutdownModule() override
+	{
+		FPropertyEditorModule& PropertyEditorModule = FModuleManager::GetModuleChecked<FPropertyEditorModule>("PropertyEditor");
+		PropertyEditorModule.UnregisterCustomClassLayout(ACaptureCharacter::StaticClass()->GetFName());
+	}
+};
+
+IMPLEMENT_MODULE(FPerformanceCaptureCoreEditorModule, PerformanceCaptureCoreEditor)
+
+#undef LOCTEXT_NAMESPACE
