MocapManager: Correcting logic surrounding changing the Source Skeletal mesh Component Reference property.

#rnx

diff --git a/Engine/Plugins/Animation/PerformanceCaptureCore/Source/PerformanceCaptureCore/Private/RetargetComponent.cpp b/Engine/Plugins/Animation/PerformanceCaptureCore/Source/PerformanceCaptureCore/Private/RetargetComponent.cpp
--- a/Engine/Plugins/Animation/PerformanceCaptureCore/Source/PerformanceCaptureCore/Private/RetargetComponent.cpp
+++ b/Engine/Plugins/Animation/PerformanceCaptureCore/Source/PerformanceCaptureCore/Private/RetargetComponent.cpp
@@ -64,7 +64,7 @@ void URetargetComponent::SetSourcePerformer(ACapturePerformer* InPerformer)
 	else
 	{
 		SourcePerformer = nullptr;
-		SetSourcePerformerMesh(Cast<USkeletalMeshComponent>(ControlledSkeletalMeshComponent.GetComponent(GetOwner())));
+		SetSourcePerformerMesh(Cast<USkeletalMeshComponent>(SourceSkeletalMeshComponent.GetComponent(GetOwner())));
 	}
 	bIsDirty = true;
 }
@@ -95,10 +95,10 @@ void URetargetComponent::SetRetargetAsset(UIKRetargeter* InRetargetAsset)
 
 			const TObjectPtr<USkeletalMeshComponent> SourceMesh = Cast<USkeletalMeshComponent>(SourceSkeletalMeshComponent.GetComponent(GetOwner()));
 		
-			if(SourceSkeletalMeshComponent.GetComponent(GetOwner()) && AnimInstance)
+			if(SourceMesh && AnimInstance)
 			{
 				AnimInstance->ConfigureAnimInstance(RetargetAsset, SourceMesh, CustomRetargetProfile);
-				ControlledMesh->InitAnim(true /*bForceReinit*/);
+				SourceMesh->InitAnim(true /*bForceReinit*/);
 			}
 			ControlledMesh->InitAnim(true /*bForceReinit*/);
 		}
@@ -172,12 +172,26 @@ void URetargetComponent::PostEditChangeProperty(FPropertyChangedEvent& PropertyC
 	{
 		if (PropertyChangedEvent.Property->GetFName() == GET_MEMBER_NAME_CHECKED(URetargetComponent, SourceSkeletalMeshComponent))
 		{
-			bIsDirty=true;
+			// OtherActor is set from the property change and should be used so we want to clear the override.
+			SetSourcePerformerMesh(nullptr);
+
+			// If the change was to zero out the OtherActor as well, we want to now force the override to self.
+			if (SourceSkeletalMeshComponent.OverrideComponent == nullptr && SourceSkeletalMeshComponent.OtherActor == nullptr)
+			{
+				SourceSkeletalMeshComponent.OverrideComponent = SourceSkeletalMeshComponent.GetComponent(GetOwner());
+			}
 		}
 
 		if (PropertyChangedEvent.Property->GetFName() == GET_MEMBER_NAME_CHECKED(URetargetComponent, ControlledSkeletalMeshComponent))
 		{
-			bIsDirty=true;
+			// OtherActor is set from the property change and should be used so we want to clear the override.
+			SetControlledMesh(nullptr);
+
+			// If the change was to zero out the OtherActor as well, we want to now force the override to self.
+			if (ControlledSkeletalMeshComponent.OverrideComponent == nullptr && ControlledSkeletalMeshComponent.OtherActor == nullptr)
+			{
+				ControlledSkeletalMeshComponent.OverrideComponent = ControlledSkeletalMeshComponent.GetComponent(GetOwner());
+			}
 		}
 
 		if (PropertyChangedEvent.Property->GetFName() == GET_MEMBER_NAME_CHECKED(URetargetComponent, bForceOtherMeshesToFollowControlledMesh))
@@ -258,7 +272,7 @@ void URetargetComponent::SetForceOtherMeshesToFollowControlledMesh(bool bInBool)
 void URetargetComponent::InitiateAnimation()
 {
 	USkeletalMeshComponent* ControlledMesh = Cast<USkeletalMeshComponent> (ControlledSkeletalMeshComponent.GetComponent(GetOwner()));
-	USkeletalMeshComponent* SourceMesh = Cast<USkeletalMeshComponent>(SourceSkeletalMeshComponent.GetComponent(nullptr));
+	USkeletalMeshComponent* SourceMesh = Cast<USkeletalMeshComponent>(SourceSkeletalMeshComponent.GetComponent(GetOwner()));
 	
 	const bool bShouldBeReinitialized = IsValid(RetargetAsset) && IsValid(ControlledMesh);
 	
