[Performance Capture Core]
-Prevent pre-requisite cycle when SourceMesh == ControlledMesh because FComponentReference wants to set default to first skelmesh component on the owning actor.

diff --git a/Engine/Plugins/Animation/PerformanceCaptureCore/Source/PerformanceCaptureCore/Private/RetargetComponent.cpp b/Engine/Plugins/Animation/PerformanceCaptureCore/Source/PerformanceCaptureCore/Private/RetargetComponent.cpp
--- a/Engine/Plugins/Animation/PerformanceCaptureCore/Source/PerformanceCaptureCore/Private/RetargetComponent.cpp
+++ b/Engine/Plugins/Animation/PerformanceCaptureCore/Source/PerformanceCaptureCore/Private/RetargetComponent.cpp
@@ -283,7 +283,7 @@ void URetargetComponent::InitiateAnimation()
 
 		TObjectPtr<URetargetAnimInstance> AnimInstance = Cast<URetargetAnimInstance>(ControlledMesh->GetAnimInstance());
 
-		if (SourceMesh && SourceMesh->GetOwner())
+		if (SourceMesh && SourceMesh->GetOwner() && SourceMesh!=ControlledMesh)
 		{
 			ControlledMesh->AddTickPrerequisiteComponent(SourceMesh);
 		}
